Add RigMapper Experimental Plugin

#virtualized

diff --git a/Engine/Plugins/Experimental/Animation/RigMapper/RigMapper.uplugin b/Engine/Plugins/Experimental/Animation/RigMapper/RigMapper.uplugin
new file mode 100644
--- /dev/null
+++ b/Engine/Plugins/Experimental/Animation/RigMapper/RigMapper.uplugin
@@ -0,0 +1,35 @@
+{
+	"FileVersion": 3,
+	"Version": 1,
+	"VersionName": "1.0",
+	"FriendlyName": "Rig Mapper",
+	"Description": "A set of animation remapping features",
+	"Category": "Animation",
+	"CreatedBy": "Epic Games, Inc.",
+	"CreatedByURL": "https://epicgames.com",
+	"DocsURL": "",
+	"MarketplaceURL": "",
+	"SupportURL": "",
+	"EnabledByDefault": false,
+	"CanContainContent": true,
+	"IsBetaVersion": false,
+	"Installed": false,
+	"Plugins": [
+		{
+			"Name": "ControlRig",
+			"Enabled": true
+		}
+	],
+	"Modules": [
+		{
+			"Name": "RigMapperEditor",
+			"Type": "UncookedOnly",
+			"LoadingPhase": "PreDefault"
+		},
+		{
+			"Name": "RigMapper",
+			"Type": "Runtime",
+			"LoadingPhase": "PreDefault"
+		}
+	]
+}
diff --git a/Engine/Plugins/Experimental/Animation/RigMapper/Source/RigMapper/Private/AnimNode_RigMapper.cpp b/Engine/Plugins/Experimental/Animation/RigMapper/Source/RigMapper/Private/AnimNode_RigMapper.cpp
new file mode 100644
--- /dev/null
+++ b/Engine/Plugins/Experimental/Animation/RigMapper/Source/RigMapper/Private/AnimNode_RigMapper.cpp
@@ -0,0 +1,279 @@
+// Copyright Epic Games, Inc. All Rights Reserved.
+
+#include "AnimNode_RigMapper.h"
+
+#include "UObject/SoftObjectPath.h"
+#include "RigMapperDefinition.h"
+#include "Engine/SkeletalMesh.h"
+#include "Animation/AnimInstanceProxy.h"
+
+DEFINE_STAT(STAT_RigMapperEvaluateNode);
+DEFINE_STAT(STAT_RigMapperEvaluateRig);
+DEFINE_STAT(STAT_RigMapperEvaluateAllRigs);
+DEFINE_STAT(STAT_RigMapperUpdateValues);
+
+FAnimNode_RigMapper::FAnimNode_RigMapper()
+{
+}
+
+FAnimNode_RigMapper::~FAnimNode_RigMapper()
+{
+	
+}
+
+void FAnimNode_RigMapper::OnInitializeAnimInstance(const FAnimInstanceProxy* InProxy,
+                                                   const UAnimInstance* InAnimInstance)
+{
+	FAnimNode_Base::OnInitializeAnimInstance(InProxy, InAnimInstance);
+}
+
+void FAnimNode_RigMapper::Initialize_AnyThread(const FAnimationInitializeContext& Context)
+{
+	DECLARE_SCOPE_HIERARCHICAL_COUNTER_ANIMNODE(Initialize_AnyThread)
+	FAnimNode_Base::Initialize_AnyThread(Context);
+
+	//GetEvaluateGraphExposedInputs().Execute(Context);
+	
+	SourcePose.Initialize(Context);
+}
+
+void FAnimNode_RigMapper::GatherDebugData(FNodeDebugData& DebugData)
+{
+	DECLARE_SCOPE_HIERARCHICAL_COUNTER_ANIMNODE(GatherDebugData)
+	FString DebugLine = DebugData.GetNodeName(this);
+	
+	DebugLine += FString::Printf(TEXT("(Alpha: %f, Definitions: "), Alpha);
+	for (int32 DefinitionIndex = 0; DefinitionIndex < Definitions.Num(); DefinitionIndex++)
+	{
+		DebugLine += FString::Printf(TEXT("%d: %s"), DefinitionIndex, Definitions[DefinitionIndex] ? *Definitions[DefinitionIndex]->GetName() : TEXT("NULL"));
+	}
+	DebugLine += ")";
+	
+	DebugData.AddDebugItem(DebugLine);
+}
+
+void FAnimNode_RigMapper::PreUpdate(const UAnimInstance* InAnimInstance)
+{
+	FAnimNode_Base::PreUpdate(InAnimInstance);
+	
+	bool bReInit = false;
+	TArray<TObjectPtr<URigMapperDefinition>>* DefinitionsToCheck = &Definitions;
+
+	if (LoadedUserData)
+	{
+		DefinitionsToCheck = &LoadedUserData->Definitions;
+	}
+
+	check(DefinitionsToCheck);
+	
+	if (DefinitionsToCheck->Num() != LoadedDefinitions.Num())
+	{
+		bReInit = true;
+	}
+	else
+	{
+		for (int32 DefinitionIndex = 0; DefinitionIndex < DefinitionsToCheck->Num(); DefinitionIndex++)
+		{
+			if ((*DefinitionsToCheck)[DefinitionIndex] != LoadedDefinitions[DefinitionIndex])
+			{
+				bReInit = true;
+				break;
+			}
+		}	
+	}
+	if (bReInit)
+	{
+		USkeletalMesh* TargetMesh = InAnimInstance && InAnimInstance->GetSkelMeshComponent() ? InAnimInstance->GetSkelMeshComponent()->GetSkeletalMeshAsset() : nullptr;
+		InitializeRigMapping(TargetMesh);		
+	}
+	// todo: Allow override of user data (from BP etc...)
+}
+
+void FAnimNode_RigMapper::Update_AnyThread(const FAnimationUpdateContext& Context)
+{
+	DECLARE_SCOPE_HIERARCHICAL_COUNTER_ANIMNODE(Update_AnyThread)
+
+	GetEvaluateGraphExposedInputs().Execute(Context);
+	SourcePose.Update(Context);
+
+	// todo: update cache if needed (config or skeleton curves changed)
+}
+
+void FAnimNode_RigMapper::CacheBones_AnyThread(const FAnimationCacheBonesContext& Context)
+{
+	DECLARE_SCOPE_HIERARCHICAL_COUNTER_ANIMNODE(CacheBones_AnyThread)
+
+	FAnimNode_Base::CacheBones_AnyThread(Context);
+
+	SourcePose.CacheBones(Context);
+
+	FBoneContainer& RequiredBones = Context.AnimInstanceProxy->GetRequiredBones();
+	if (RequiredBones.IsValid())
+	{
+		// todo: cache bones if required
+		// const FCompactPoseBoneIndex CompactPoseBoneIndex = RequiredBones.MakeCompactPoseIndex();
+	}
+		
+	USkeletalMesh* TargetMesh = Context.AnimInstanceProxy->GetSkelMeshComponent() ? Context.AnimInstanceProxy->GetSkelMeshComponent()->GetSkeletalMeshAsset() : nullptr;
+	InitializeRigMapping(TargetMesh);		
+}
+
+void FAnimNode_RigMapper::Evaluate_AnyThread(FPoseContext& Output)
+{
+	DECLARE_SCOPE_HIERARCHICAL_COUNTER_ANIMNODE(Evaluate_AnyThread)
+	
+	if (IsLODEnabled(Output.AnimInstanceProxy))
+	{
+		SCOPE_CYCLE_COUNTER(STAT_RigMapperEvaluateNode);
+		
+		SourcePose.Evaluate(Output);
+
+		if (Alpha > 0.f)
+		{
+			EvaluateRigMapping(Output);
+		}
+	}
+}
+
+// todo: handle bones
+void FAnimNode_RigMapper::EvaluateRigMapping(FPoseContext& Output) const
+{
+	SCOPE_CYCLE_COUNTER(STAT_RigMapperEvaluateAllRigs);
+	
+	for (int i = 0; i < RigMappers.Num(); i++)
+	{
+		const TSharedPtr<FacialRigMapping::FRigMapper> RigMapper = RigMappers[i];
+		RigMapper->SetDirty();
+
+		// Set first rig mapper's input values from current pose
+		if (i == 0)
+		{
+			SCOPE_CYCLE_COUNTER(STAT_RigMapperUpdateValues);
+			
+			const TArray<FString>& InputNames = RigMapper->GetInputNames();
+			for (int n = 0; n < InputNames.Num(); n++)
+			{
+				RigMapper->SetDirectValue(n, Output.Curve.Get(*InputNames[n]));
+			}
+		}
+		else
+		{
+			// Other RigMappers: Link/Map prev rig mapper outputs to current inputs
+			check(IndexMaps.IsValidIndex(i - 1))
+			const FRigMapperIndexMap& IndexMapping = IndexMaps[i - 1];
+
+			// Get output values from previous rig mapper
+			TArray<double> PrevOutputValuesInOrder;
+			{
+				SCOPE_CYCLE_COUNTER(STAT_RigMapperEvaluateRig);
+				PrevOutputValuesInOrder = RigMappers[i - 1]->GetOutputValuesInOrder();
+			}
+
+			// Set input values for current rig mapper
+			{
+				SCOPE_CYCLE_COUNTER(STAT_RigMapperUpdateValues);
+				
+				for (int32 n = 0; n < IndexMapping.Num(); n++)
+				{
+					if (const int32 Index = IndexMapping[n]; Index != INDEX_NONE)
+					{
+						RigMapper->SetDirectValue(n, PrevOutputValuesInOrder[Index]);
+					}
+				}
+			}
+		}
+	}
+
+	// Set current pose from last rig mapper's output values 
+	if (!RigMappers.IsEmpty())
+	{
+		TMap<FString, double> OutputValues;
+		{
+			SCOPE_CYCLE_COUNTER(STAT_RigMapperEvaluateRig);
+			OutputValues = RigMappers.Last()->GetOutputValues();
+		}
+		
+		for (const TPair<FString, double>& OutputValue : OutputValues)
+		{
+			const float Value = FMath::Lerp(Output.Curve.Get(*OutputValue.Key), OutputValue.Value, Alpha);
+			Output.Curve.Set(*OutputValue.Key, Value);
+		}
+	}
+}
+
+int32 FAnimNode_RigMapper::GetLODThreshold() const
+{
+	return LODThreshold;
+}
+
+bool FAnimNode_RigMapper::InitializeRigMapping(USkeletalMesh* TargetMesh)
+{
+	TArray<URigMapperDefinition*> DefinitionsToLoad;
+
+	if (TargetMesh)
+	{
+		LoadedUserData = Cast<URigMapperDefinitionUserData>(TargetMesh->GetAssetUserDataOfClass(URigMapperDefinitionUserData::StaticClass()));
+		if (LoadedUserData)
+		{
+			DefinitionsToLoad = LoadedUserData->Definitions;
+		}
+	}
+	if (DefinitionsToLoad.IsEmpty())
+	{
+		DefinitionsToLoad = Definitions;
+	}
+
+	if (DefinitionsToLoad.IsEmpty())
+	{
+		return false;
+	}
+	for (int i = 0; i < DefinitionsToLoad.Num(); i++)
+	{
+		if (!DefinitionsToLoad[i])
+		{
+			UE_LOG(LogTemp, Error, TEXT("Invalid definition at index %d"), i)
+			return false;
+		}
+	}
+	RigMappers.Empty(DefinitionsToLoad.Num());
+	IndexMaps.Empty(DefinitionsToLoad.Num() - 1);
+	
+	for (const URigMapperDefinition* Definition : DefinitionsToLoad)
+	{
+		TSharedPtr<FacialRigMapping::FRigMapper> RigMapper = MakeShared<FacialRigMapping::FRigMapper>();
+	
+		if (!RigMapper->LoadDefinition(Definition))
+		{
+			UE_LOG(LogTemp, Error, TEXT("Could not load definition %s"), *Definition->GetPathName())
+			RigMappers.Empty();
+			return false;
+		}
+
+		if (RigMappers.Num() > 0)
+		{
+			IndexMaps.Add(MakeIndexMap(RigMappers.Last(), RigMapper));
+		}
+
+		RigMappers.Add(RigMapper);
+	}
+
+	LoadedDefinitions = MoveTemp(DefinitionsToLoad);
+	return true;
+}
+
+FRigMapperIndexMap FAnimNode_RigMapper::MakeIndexMap(const TSharedPtr<FacialRigMapping::FRigMapper>& OutputRig, const TSharedPtr<FacialRigMapping::FRigMapper>& InputRig)
+{
+	const TArray<FString> FromOutputs = OutputRig->GetOutputNames();
+	const TArray<FString> ToInputs = InputRig->GetInputNames();
+				
+	FRigMapperIndexMap IndexMapping;
+	
+	IndexMapping.Reserve(ToInputs.Num());
+
+	for (const FString& Input : ToInputs)
+	{
+		IndexMapping.Add(FromOutputs.Find(Input));
+	}
+
+	return IndexMapping;
+}
\ No newline at end of file
diff --git a/Engine/Plugins/Experimental/Animation/RigMapper/Source/RigMapper/Private/RigMapper.cpp b/Engine/Plugins/Experimental/Animation/RigMapper/Source/RigMapper/Private/RigMapper.cpp
new file mode 100644
--- /dev/null
+++ b/Engine/Plugins/Experimental/Animation/RigMapper/Source/RigMapper/Private/RigMapper.cpp
@@ -0,0 +1,379 @@
+// Copyright Epic Games, Inc. All Rights Reserved.
+
+#include "RigMapper.h"
+
+#include "RigMapperDefinition.h"
+#include "UObject/SoftObjectPath.h"
+
+namespace FacialRigMapping
+{
+	bool FRigMapper::IsValid() const
+	{
+		return !Nodes.IsEmpty() && !OutputNodes.IsEmpty() && OutputNodes.Num() == OutputNames.Num() && !InputNodes.IsEmpty() && InputNodes.Num() == InputNames.Num();
+	}
+	
+	void FRigMapper::Reset()
+	{
+		Nodes.Empty();
+    		
+    	OutputNodes.Empty();
+    	OutputNames.Empty();
+    		
+    	InputNodes.Empty();
+    	InputNames.Empty();
+	}
+	
+	bool FRigMapper::LoadDefinition(const URigMapperDefinition* Definition)
+	{
+		Reset();
+		
+		for (const FString& InputName : Definition->Inputs)
+		{
+			TSharedPtr<FEvalInputNode> Node = MakeShared<FEvalInputNode>(InputName);
+			
+			InputNodes.Add(Node);
+			InputNames.Add(InputName);
+
+			Nodes.Add(InputName, Node);
+		}
+
+		for (const FRigMapperMultiplyFeature& Feature : Definition->Features.Multiply)
+		{
+			Nodes.Add(Feature.Name, MakeShared<FEvalNodeMultiply>(Feature));
+		}
+		for (const FRigMapperWsFeature& Feature : Definition->Features.WeightedSums)
+		{
+			Nodes.Add(Feature.Name, MakeShared<FEvalNodeWeightedSum>(Feature));
+		}
+		for (const FRigMapperSdkFeature& Feature : Definition->Features.SDKs)
+		{
+			Nodes.Add(Feature.Name, MakeShared<FEvalNodePiecewiseLinear>(Feature));
+		}
+
+		for (const TPair<FString, FString>& Output : Definition->Outputs)
+		{
+			const FString OutputName = Output.Key;
+			const FString LinkedNodeName = Output.Value;
+			
+			if (const EvalPtr* Node = Nodes.Find(LinkedNodeName))
+			{
+				OutputNodes.Add(*Node);
+				OutputNames.Add(OutputName);
+			}
+		}
+
+		if (IsValid())
+		{
+			for (const TPair<FString, EvalPtr>& NamedNode : Nodes)
+			{
+				NamedNode.Value->Initialize(Nodes);
+			}
+		}
+		return IsValid();
+	}
+
+	bool FRigMapper::SetDirectValue(const int32 InputIndex, double Value)
+	{
+		if (InputNodes.IsValidIndex(InputIndex))
+		{
+			InputNodes[InputIndex]->SetDirect(Value);
+			return true;
+		}
+		return false;
+	}
+	
+	bool FRigMapper::SetDirectValue(const FString& InputName, double Value)
+	{
+		return SetDirectValue(InputNames.Find(InputName), Value);
+	}
+
+	TMap<FString, double> FRigMapper::GetOutputValues(bool bSkipUnset) const
+	{
+		TMap<FString, double> OutValues;
+	
+		for (int i = 0; i < OutputNodes.Num(); i++)
+		{
+			double OutputValue = 0.0f;
+			
+			if (OutputNodes[i]->TryGetValue(OutputValue) || !bSkipUnset)
+			{
+				OutValues.Add(OutputNames[i], OutputValue);
+			}
+		}
+
+		return OutValues;
+	}
+
+	const TArray<FString>& FRigMapper::GetOutputNames() const
+	{
+		return OutputNames;
+	}
+
+	TArray<double> FRigMapper::GetOutputValuesInOrder() const
+	{
+		TArray<double> OutValues;
+	
+		for (int i = 0; i < OutputNodes.Num(); i++)
+		{
+			OutValues.Add(OutputNodes[i]->GetValue());
+		}
+
+		return OutValues;
+	}
+
+	void FRigMapper::SetDirty() const
+	{
+		for (const auto& N : Nodes)
+		{
+			N.Value->Reset();
+		}
+	}
+
+	const TArray<FString>& FRigMapper::GetInputNames() const
+	{
+		return InputNames;
+	}
+	
+	FEvalNode::FEvalNode(const FString& InName):
+		Name(InName)
+	{
+	}
+
+	double FEvalNode::GetValue() const
+	{
+		double Value = 0.0f;
+
+		TryGetValue(Value);
+		return Value;
+	}
+
+	bool FEvalNode::TryGetValue(double& OutValue) const
+	{
+		if (!IsInitialized())
+		{
+			return false;
+		}
+		if (!CachedValue.IsSet())
+		{
+			if (Evaluate(OutValue))
+			{
+				CachedValue = OutValue;
+				return true;
+			}
+			return false;
+		}
+		else
+		{
+			OutValue = CachedValue.GetValue();
+			return true;
+		}
+	}
+
+	void FEvalNode::SetDirect(double Value)
+	{
+		CachedValue = Value;
+	}
+
+	void FEvalNode::Reset()
+	{
+		CachedValue.Reset();
+	}
+
+	bool FEvalNode::IsInitialized() const
+	{
+		return bInitialized;
+	}
+
+	FEvalNodeWeightedSum::FEvalNodeWeightedSum(const FRigMapperWsFeature& FeatureDefinition):
+		FEvalNode(FeatureDefinition.Name)
+	{
+		if (FeatureDefinition.Range.bHasLowerBound)
+		{
+			Range.SetLowerBound(TRange<double>::BoundsType(FeatureDefinition.Range.LowerBound));
+		}
+		if (FeatureDefinition.Range.bHasUpperBound)
+		{
+			Range.SetUpperBound(TRange<double>::BoundsType(FeatureDefinition.Range.UpperBound));
+		}
+		WeightedLinkedInputNames = FeatureDefinition.Inputs;
+	}
+
+	bool FEvalNodeWeightedSum::Evaluate(double &OutValue) const
+	{
+		bool bInputEvaluated = false;
+		OutValue = 0;
+
+		for (const auto& WeightedInput : WeightedLinkedInputs)
+		{
+			double InputValue = 0.0f;
+			
+			if (WeightedInput.Key->TryGetValue(InputValue))
+			{
+				bInputEvaluated = true;
+			}
+			OutValue += WeightedInput.Key->GetValue() * WeightedInput.Value;
+		}
+		
+		if (Range.HasLowerBound() && OutValue < Range.GetLowerBoundValue())
+		{
+			OutValue = Range.GetLowerBoundValue();
+		}
+		if (Range.HasUpperBound() && OutValue > Range.GetUpperBoundValue())
+		{
+			OutValue = Range.GetUpperBoundValue();
+		}
+		
+		return bInputEvaluated;
+	}
+
+	void FEvalNodeWeightedSum::Initialize(const TMap<FString, EvalPtr>& Nodes)
+	{
+		FEvalNode::Initialize(Nodes);
+		WeightedLinkedInputs.Empty(WeightedLinkedInputNames.Num());
+
+		for (const TPair<FString, double>& Pair : WeightedLinkedInputNames)
+		{
+			if (const EvalPtr* Node = Nodes.Find(Pair.Key))
+			{
+				if (Node->IsValid())
+				{
+					WeightedLinkedInputs.Add(*Node, Pair.Value);
+				}
+			}
+			// todo: log error?
+		}
+
+		// Will be considered not initialized if any node could not be found
+		bInitialized = WeightedLinkedInputs.Num() == WeightedLinkedInputNames.Num();
+	}
+
+	FEvalNodePiecewiseLinear::FEvalNodePiecewiseLinear(const FRigMapperSdkFeature& FeatureDefinition):
+		FEvalNode(FeatureDefinition.Name)
+	{
+		LinkedInputName = FeatureDefinition.Input;
+
+		for (const FRigMapperSdkKey& Key : FeatureDefinition.Keys)
+		{
+			Keys.Add({ Key.In, Key.Out });
+		}
+		Keys.Sort([](const TPair<double,double>& Key1, const TPair<double,double>& Key2) { return Key1.Key < Key2.Key; });
+	}
+
+	bool FEvalNodePiecewiseLinear::Evaluate(double &OutValue) const
+	{
+		double InputValue = 0.0f;
+
+		if (!LinkedInput->TryGetValue(InputValue))
+		{
+			return false;
+		}
+
+		if (!Evaluate_Static(InputValue, Keys, OutValue))
+		{
+			UE_LOG(LogTemp, Warning, TEXT("PiecewiseLinear could not calculate the output value"))
+			return false;	
+		}
+
+		return true;
+	}
+
+	bool FEvalNodePiecewiseLinear::Evaluate_Static(const double InputValue, const TArray<TPair<double,double>>& Keys, double &OutValue)
+	{
+		if (InputValue <= Keys[0].Key)
+		{
+			OutValue = Keys[0].Value;
+			return true;
+		}
+	
+		if (InputValue >= Keys.Last().Key)
+		{
+			OutValue = Keys.Last().Value;
+			return true;
+		}
+	
+		for (int32 KeyIndex = 1 ; KeyIndex < Keys.Num(); KeyIndex++)
+		{
+			const TPair<double, double>& CurrentKey = Keys[KeyIndex];
+			
+			if (InputValue == CurrentKey.Key)
+			{
+				OutValue = CurrentKey.Value;
+				return true;
+			}
+			if (InputValue < CurrentKey.Key)
+			{
+				const TPair<double, double>& PrevKey = Keys[KeyIndex - 1];
+				
+				const double Percent = (InputValue - PrevKey.Key) / (CurrentKey.Key - PrevKey.Key);
+
+				OutValue = PrevKey.Value + Percent * (CurrentKey.Value - PrevKey.Value);
+				return true;
+			}
+		}
+
+		return false;
+	}
+
+	void FEvalNodePiecewiseLinear::Initialize(const TMap<FString, EvalPtr>& Nodes)
+	{
+		FEvalNode::Initialize(Nodes);
+		LinkedInput.Reset();
+		
+		if (const EvalPtr* Node = Nodes.Find(LinkedInputName))
+		{
+			LinkedInput = *Node;
+		}
+
+		bInitialized = LinkedInput.IsValid();
+	}
+
+	FEvalNodeMultiply::FEvalNodeMultiply(const FRigMapperMultiplyFeature& FeatureDefinition):
+		FEvalNode(FeatureDefinition.Name), LinkedInputNames(FeatureDefinition.Inputs)
+	{
+	}
+
+	bool FEvalNodeMultiply::Evaluate(double &OutValue) const
+	{
+		if (LinkedInputs.IsEmpty())
+		{
+			return false;
+		}
+
+		OutValue = 1.0;
+		bool bInputEvaluated = false;
+
+		for (const EvalPtr& Input : LinkedInputs)
+		{
+			double InputValue = 0.0f;
+
+			if (Input->TryGetValue(InputValue))
+			{
+				bInputEvaluated = true;
+			}
+			OutValue *= Input->GetValue();
+		}
+
+		return bInputEvaluated;
+	}
+
+	void FEvalNodeMultiply::Initialize(const TMap<FString, EvalPtr>& Nodes)
+	{
+		FEvalNode::Initialize(Nodes);
+		LinkedInputs.Empty(LinkedInputNames.Num());
+
+		for (const FString& InputName : LinkedInputNames)
+		{
+			if (const EvalPtr* Node = Nodes.Find(InputName))
+			{
+				if (Node->IsValid())
+				{
+					LinkedInputs.Add(*Node);
+				}
+			}
+			// todo: log error?
+		}
+
+		// Will be considered not initialized if any node could not be found
+		bInitialized = LinkedInputs.Num() == LinkedInputNames.Num();
+	}
+}
\ No newline at end of file
diff --git a/Engine/Plugins/Experimental/Animation/RigMapper/Source/RigMapper/Private/RigMapperDefinition.cpp b/Engine/Plugins/Experimental/Animation/RigMapper/Source/RigMapper/Private/RigMapperDefinition.cpp
new file mode 100644
--- /dev/null
+++ b/Engine/Plugins/Experimental/Animation/RigMapper/Source/RigMapper/Private/RigMapperDefinition.cpp
@@ -0,0 +1,1150 @@
+// Copyright Epic Games, Inc. All Rights Reserved.
+
+#include "RigMapperDefinition.h"
+
+#include "RigMapper.h"
+#include "Misc/FileHelper.h"
+#include "Serialization/JsonTypes.h"
+#include "Dom/JsonValue.h"
+#include "Dom/JsonObject.h"
+#include "Serialization/JsonReader.h"
+#include "Serialization/JsonWriter.h"
+#include "Serialization/JsonSerializer.h"
+
+DEFINE_LOG_CATEGORY(LogRigMapperDefinition)
+
+bool FRigMapperFeature::IsValid(const TArray<FString>& InputNames, bool bWarn) const
+{
+	bool bValid = true;
+	
+	if (Name.IsEmpty())
+	{
+		UE_LOG(LogRigMapperDefinition, Warning, TEXT("Invalid (empty) feature name"))
+		bValid = false;
+	}
+
+	TArray<FString> FeatureInputs;
+	GetInputs(FeatureInputs);
+	
+	if (FeatureInputs.Contains(Name))
+	{
+		UE_LOG(LogRigMapperDefinition, Warning, TEXT("Feature %s is referencing itself"), *Name)
+		bValid = false;
+	}
+
+	if (FeatureInputs.IsEmpty())
+	{
+		UE_LOG(LogRigMapperDefinition, Warning, TEXT("Feature %s does not reference any input"), *Name)
+		bValid = false;
+	}
+	
+	for (const FString& Input : FeatureInputs)
+	{
+		if (!InputNames.Contains(Input))
+		{
+			UE_LOG(LogRigMapperDefinition, Warning, TEXT("Undefined input or feature %s referenced in feature %s"), *Input, *Name)
+			bValid = false;
+		}
+	}
+	
+	return bValid;
+}
+
+bool FRigMapperFeature::GetJsonArray(TSharedPtr<FJsonObject> JsonObject, TArray<TSharedPtr<FJsonValue>>& OutArray, const FString& Identifier, const FString& OwnerIdentifier) const
+{
+	if (!OwnerIdentifier.IsEmpty())
+	{
+		if (!JsonObject->HasTypedField<EJson::Object>(*OwnerIdentifier))
+		{
+			UE_LOG(LogRigMapperDefinition, Warning, TEXT("Missing '%s' field for feature %s"), *OwnerIdentifier, *Name)
+			return false;
+		}
+		JsonObject = JsonObject->GetObjectField(OwnerIdentifier);
+	}
+	if (!JsonObject->HasTypedField<EJson::Array>(*Identifier))
+	{
+		UE_LOG(LogRigMapperDefinition, Warning, TEXT("Missing '%s' field for feature %s"), *Identifier, *Name)
+		return false;
+	}
+	OutArray = JsonObject->GetArrayField(*Identifier);
+	return true;
+}
+
+bool FRigMapperMultiplyFeature::LoadFromJsonObject(const TSharedPtr<FJsonObject>& JsonObject)
+{
+	TArray<TSharedPtr<FJsonValue>> InputFeatures;
+	
+	if (!GetJsonArray(JsonObject, InputFeatures, TEXT("input_features")))
+	{
+		return false;
+	}
+	if (InputFeatures.Num() < 2)
+	{
+		UE_LOG(LogRigMapperDefinition, Warning, TEXT("Feature %s does not reference enough input"), *Name)
+		return false;
+	}
+	for (const auto& InputFeature : InputFeatures)
+	{
+		Inputs.Add(InputFeature->AsString());
+	}
+	
+	return true;
+}
+
+bool FRigMapperMultiplyFeature::IsValid(const TArray<FString>& InputNames, bool bWarn) const
+{
+	bool bValid = FRigMapperFeature::IsValid(InputNames, bWarn);
+	
+	if (bValid && Inputs.Num() < 2)
+	{
+		UE_LOG(LogRigMapperDefinition, Warning, TEXT("Feature %s does not reference enough input"), *Name)
+		return false;
+	}
+	
+	return bValid;
+}
+
+bool FRigMapperWsFeature::LoadFromJsonObject(const TSharedPtr<FJsonObject>& JsonObject)
+{
+	TArray<TSharedPtr<FJsonValue>> JsonInputs;
+	if (!GetJsonArray(JsonObject, JsonInputs, TEXT("input_features")))
+	{
+		return false;
+	}
+	
+	TArray<TSharedPtr<FJsonValue>> JsonWeights;
+	if (!GetJsonArray(JsonObject, JsonWeights, TEXT("weights"), TEXT("params")))
+	{
+		return false;
+	}
+	if (JsonWeights.Num() != JsonInputs.Num())
+	{
+		UE_LOG(LogRigMapperDefinition, Warning, TEXT("Number of inputs does not match number of weights for feature %s"), *Name)
+		return false;
+	}
+	for (int32 InputIndex = 0; InputIndex < JsonInputs.Num(); InputIndex++)
+	{
+		Inputs.Add(JsonInputs[InputIndex]->AsString(), JsonWeights[InputIndex]->AsNumber());
+	}
+
+	Range.bHasLowerBound = JsonObject->GetObjectField(TEXT("params"))->HasField(TEXT("min"));
+	if (Range.bHasLowerBound)
+	{
+		Range.LowerBound = JsonObject->GetObjectField(TEXT("params"))->GetNumberField(TEXT("min"));
+	}
+	Range.bHasUpperBound = JsonObject->GetObjectField(TEXT("params"))->HasField(TEXT("max"));
+	if (Range.bHasUpperBound)
+	{
+		Range.UpperBound = JsonObject->GetObjectField(TEXT("params"))->GetNumberField(TEXT("max"));
+	}
+	if (Range.bHasLowerBound && Range.bHasUpperBound && Range.LowerBound > Range.UpperBound)
+	{
+		UE_LOG(LogRigMapperDefinition, Warning, TEXT("Invalid range for feature %s"), *Name)
+		return false;
+	}
+	return true;
+}
+
+bool FRigMapperWsFeature::IsValid(const TArray<FString>& InputNames, bool bWarn) const
+{
+	bool bValid = FRigMapperFeature::IsValid(InputNames, bWarn);
+
+	double TotalWeight = 0;
+	for (const TPair<FString, double>& InputWeight : Inputs)
+	{
+		TotalWeight += InputWeight.Value;
+	}
+	if (bWarn && TotalWeight == 0)
+	{
+		UE_LOG(LogRigMapperDefinition, Warning, TEXT("Total Weights for feature %s add up to 0"), *Name)
+	}
+	if (bWarn && (TotalWeight < -1 || TotalWeight > 1) && (!Range.bHasLowerBound || !Range.bHasUpperBound))
+	{
+		UE_LOG(LogRigMapperDefinition, Warning, TEXT("Total Weights for feature %s are quite high (%f) even though a range was not set"), *Name, TotalWeight)
+	}
+	if (Range.bHasLowerBound && Range.bHasUpperBound && Range.LowerBound > Range.UpperBound)
+	{
+		UE_LOG(LogRigMapperDefinition, Warning, TEXT("Range of [%f-%f] for feature %s is invalid"), Range.LowerBound, Range.UpperBound, *Name)
+		bValid = false;
+	}
+	
+	return bValid;
+}
+
+void FRigMapperWsFeature::GetInputs(TArray<FString>& OutInputs) const
+{
+	Inputs.GenerateKeyArray(OutInputs);
+}
+
+bool FRigMapperSdkFeature::LoadFromJsonObject(const TSharedPtr<FJsonObject>& JsonObject)
+{
+	TArray<TSharedPtr<FJsonValue>> JsonInputs;
+	if (!GetJsonArray(JsonObject, JsonInputs, TEXT("input_features")))
+	{
+		return false;
+	}
+	if (JsonInputs.Num() != 1)
+	{
+		UE_LOG(LogRigMapperDefinition, Warning, TEXT("Sdk feature %s should have a single element in the 'input_features' array field"), *Name)
+		return false;
+	}
+	Input = JsonInputs[0]->AsString();
+
+	TArray<TSharedPtr<FJsonValue>> InValues;
+	if (!GetJsonArray(JsonObject, InValues, TEXT("in_val"), TEXT("params")))
+	{
+		return false;
+	}
+
+	TArray<TSharedPtr<FJsonValue>> OutValues;
+	if (!GetJsonArray(JsonObject, OutValues, TEXT("out_val"), TEXT("params")))
+	{
+		return false;
+	}
+
+	if (InValues.Num() < 2)
+	{
+		UE_LOG(LogRigMapperDefinition, Warning, TEXT("Not enough keys for SDK feature %s (expected minimum 2, got %d)"), *Name, InValues.Num())
+	}
+	if (InValues.Num() != OutValues.Num())
+	{
+		UE_LOG(LogRigMapperDefinition, Warning, TEXT("Number of input values does not match number of output values for feature %s"), *Name)
+		return false;
+	}
+	for (int32 ValueIndex = 0 ; ValueIndex < InValues.Num(); ValueIndex++)
+	{
+		Keys.Add({ InValues[ValueIndex]->AsNumber(), OutValues[ValueIndex]->AsNumber() });
+	}
+	Keys.Sort([](const FRigMapperSdkKey& Key1, const FRigMapperSdkKey Key2) { return Key1.In < Key2.In; });
+	
+	return true;
+}
+
+bool FRigMapperSdkFeature::IsValid(const TArray<FString>& InputNames, bool bWarn) const
+{
+	bool bValid = FRigMapperFeature::IsValid(InputNames, bWarn);
+
+	if (Keys.Num() < 2)
+	{
+		UE_LOG(LogRigMapperDefinition, Warning, TEXT("Not enough keys for SDK feature %s (expected minimum 2, got %d)"), *Name, Keys.Num())
+		bValid = false;
+	}
+	
+	return bValid;
+}
+
+bool FRigMapperFeatureDefinitions::AddFromJsonObject(const FString& FeatureName, const TSharedPtr<FJsonObject>& JsonObject)
+{
+	if (!JsonObject->HasTypedField<EJson::String>(TEXT("type")))
+	{
+		UE_LOG(LogRigMapperDefinition, Warning, TEXT("Missing 'type' field for feature %s"), *FeatureName)
+		return false;
+	}
+	const FString FeatureType = JsonObject->GetStringField(TEXT("type"));
+
+	bool bValidFeature = false;
+	
+	if (FeatureType == "weighted_sum")
+	{
+		WeightedSums.Add(FRigMapperWsFeature(FeatureName));
+		bValidFeature = WeightedSums.Last().LoadFromJsonObject(JsonObject);
+	}
+	else if (FeatureType == "sdk")
+	{
+		SDKs.Add(FRigMapperSdkFeature(FeatureName));
+		bValidFeature = SDKs.Last().LoadFromJsonObject(JsonObject);
+	}
+	else if (FeatureType == "multiply")
+	{
+		Multiply.Add(FRigMapperMultiplyFeature(FeatureName));
+		bValidFeature = Multiply.Last().LoadFromJsonObject(JsonObject);
+	}
+	else
+	{
+		UE_LOG(LogRigMapperDefinition, Warning, TEXT("Invalid type for feature %s (%s)"), *FeatureName, *FeatureType)
+		return bValidFeature;
+	}
+	return bValidFeature;
+}
+
+bool FRigMapperFeatureDefinitions::GetFeatureNames(TArray<FString>& OutFeatureNames) const
+{
+	bool bFoundDuplicate = false;
+
+	// todo: from interface
+	for (const FRigMapperMultiplyFeature& Feature : Multiply)
+	{
+		bFoundDuplicate |= OutFeatureNames.Contains(Feature.Name);
+		OutFeatureNames.Add(Feature.Name);
+	}
+	for (const FRigMapperSdkFeature& Feature : SDKs)
+	{
+		bFoundDuplicate |= OutFeatureNames.Contains(Feature.Name);
+		OutFeatureNames.Add(Feature.Name);
+	}
+	for (const FRigMapperWsFeature& Feature : WeightedSums)
+	{
+		bFoundDuplicate |= OutFeatureNames.Contains(Feature.Name);
+		OutFeatureNames.Add(Feature.Name);
+	}
+	return !bFoundDuplicate;
+}
+
+bool FRigMapperFeatureDefinitions::IsValid(const TArray<FString>& InputNames, bool bWarn) const
+{
+	TArray<FString> FeatureNames;
+	bool bValid = true;
+
+	const bool bDuplicateFeatureName = GetFeatureNames(FeatureNames);
+	TArray<FString> CheckFeatureNames;
+	for (const FString& FeatureName : FeatureNames)
+	{
+		if (InputNames.Contains(FeatureName))
+		{
+			UE_LOG(LogRigMapperDefinition, Warning, TEXT("Conflicting input and feature name: %s"), *FeatureName)
+			bValid = false;
+		}
+		if (bDuplicateFeatureName)
+		{
+			if (CheckFeatureNames.Contains(FeatureName))
+			{
+				UE_LOG(LogRigMapperDefinition, Warning, TEXT("Duplicate feature name: %s"), *FeatureName)
+				bValid = false;
+			}
+			else
+			{
+				CheckFeatureNames.Add(FeatureName);
+			}
+		}
+	}
+
+	TArray<FString> FeatureAndInputNames = MoveTemp(FeatureNames);
+	FeatureAndInputNames.Append(InputNames);
+
+	for (const FRigMapperMultiplyFeature& Feature : Multiply)
+	{
+		bValid &= Feature.IsValid(FeatureAndInputNames, bWarn);
+	}
+	for (const FRigMapperSdkFeature& Feature : SDKs)
+	{
+		bValid &= Feature.IsValid(FeatureAndInputNames, bWarn);
+	}
+	for (const FRigMapperWsFeature& Feature : WeightedSums)
+	{
+		bValid &= Feature.IsValid(FeatureAndInputNames, bWarn);
+	}
+	
+	return bValid;
+}
+
+void FRigMapperFeatureDefinitions::Empty()
+{
+	Multiply.Empty();
+	WeightedSums.Empty();
+	SDKs.Empty();
+}
+
+FRigMapperFeature* FRigMapperFeatureDefinitions::Find(const FString& FeatureName, ERigMapperFeatureType& OutFeatureType)
+{
+	if (FRigMapperMultiplyFeature* Feature = Multiply.FindByPredicate([FeatureName](const FRigMapperMultiplyFeature& Other) { return Other.Name == FeatureName; }))
+	{
+		OutFeatureType = ERigMapperFeatureType::Multiply;
+		return Feature;
+	}
+
+	if (FRigMapperWsFeature* Feature = WeightedSums.FindByPredicate([FeatureName](const FRigMapperWsFeature& Other) { return Other.Name == FeatureName; }))
+	{
+		OutFeatureType = ERigMapperFeatureType::WeightedSum;
+		return Feature;
+	}
+
+	if (FRigMapperSdkFeature* Feature = SDKs.FindByPredicate([FeatureName](const FRigMapperSdkFeature& Other) { return Other.Name == FeatureName; }))
+	{
+		OutFeatureType = ERigMapperFeatureType::SDK;
+		return Feature;
+	}
+
+	return nullptr;
+}
+
+bool URigMapperDefinition::LoadFromJsonFile(const FFilePath& JsonFilePath)
+{
+	FString JsonAsString;
+		
+	if (!FFileHelper::LoadFileToString(JsonAsString, *JsonFilePath.FilePath))
+	{
+		UE_LOG(LogRigMapperDefinition, Warning, TEXT("Could not open json file"))
+		return false;
+	}
+
+	return LoadFromJsonString(JsonAsString);
+}
+
+void URigMapperDefinition::Empty()
+{
+	Modify();
+	bDefinitionValidated = false;
+	Inputs.Empty();
+	Outputs.Empty();
+	Features.Empty();
+	NullOutputs.Empty();
+}
+
+bool URigMapperDefinition::LoadFromJsonString(const FString& JsonString)
+{
+	TSharedPtr<FJsonObject> JsonObject = MakeShareable(new FJsonObject());
+	const TSharedRef<TJsonReader<>> JsonReader = TJsonReaderFactory<>::Create(JsonString);
+	
+	if (!FJsonSerializer::Deserialize(JsonReader, JsonObject) && JsonObject.IsValid())
+	{
+		UE_LOG(LogRigMapperDefinition, Warning, TEXT("Could not deserialize json data to load definition"))
+		return false;
+	}
+
+	Empty();
+	bDefinitionValidated = true;
+	
+	TArray<FString> FeatureNames;
+
+	const bool bValidInputs = LoadInputsFromJsonObject(JsonObject);
+	if (!bValidInputs)
+	{
+		bDefinitionValidated = false;
+	}
+	const bool bValidFeatures = LoadFeaturesFromJsonObject(JsonObject, FeatureNames);
+	if (!bValidFeatures)
+	{
+		bDefinitionValidated = false;
+	}
+	const bool bValidOutputs = LoadOutputsFromJsonObject(JsonObject, FeatureNames);
+	if (!bValidOutputs)
+	{
+		bDefinitionValidated = false;
+	}
+	const bool bValidNullOutputs = LoadNullOutputsFromJsonObject(JsonObject);
+	if (!bValidNullOutputs)
+	{
+		bDefinitionValidated = false;
+	}
+
+	OnRigMapperDefinitionUpdated.Broadcast();
+
+	return bDefinitionValidated;
+}
+
+bool URigMapperDefinition::LoadInputsFromJsonObject(const TSharedPtr<FJsonObject>& JsonObject)
+{
+	if (!JsonObject->HasTypedField(TEXT("inputs"), EJson::Array))
+	{
+		UE_LOG(LogRigMapperDefinition, Warning, TEXT("Missing inputs field"))
+		bDefinitionValidated = false;
+		return false;
+	}
+	for (const auto& InputAttr : JsonObject->GetArrayField(TEXT("inputs")))
+	{
+		const FString InputString = InputAttr->AsString();
+		if (Inputs.Contains(InputString))
+		{
+			UE_LOG(LogRigMapperDefinition, Warning, TEXT("Duplicate input was found: %s"), *InputString)
+		}
+		else
+		{
+			Inputs.Add(InputString);	
+		}
+	}
+	if (Inputs.IsEmpty())
+	{
+		UE_LOG(LogRigMapperDefinition, Warning, TEXT("Not enough inputs"))
+		bDefinitionValidated = false;
+	}
+	return bDefinitionValidated;
+}
+
+bool URigMapperDefinition::LoadNullOutputsFromJsonObject(const TSharedPtr<FJsonObject>& JsonObject)
+{
+	if (JsonObject->HasTypedField(TEXT("null_outputs"), EJson::Array))
+	{
+		for (const auto& NullOutputAttr : JsonObject->GetArrayField(TEXT("null_outputs")))
+		{
+			const FString NullOutputString = NullOutputAttr->AsString();
+			if (NullOutputs.Contains(NullOutputString))
+			{
+				UE_LOG(LogRigMapperDefinition, Warning, TEXT("Duplicate null output was found: %s"), *NullOutputString)
+			}
+			else
+			{
+				if (Outputs.Contains(NullOutputString))
+				{
+					UE_LOG(LogRigMapperDefinition, Warning, TEXT("Null output conflicts with existing output: %s"), *NullOutputString)
+					bDefinitionValidated = false;
+				}
+				else
+				{
+					NullOutputs.Add(NullOutputString);	
+				}
+			}
+		}
+	}
+	return bDefinitionValidated;
+}
+
+bool URigMapperDefinition::LoadFeaturesFromJsonObject(const TSharedPtr<FJsonObject>& JsonObject, TArray<FString>& FeatureNames)
+{
+	if (!JsonObject->HasTypedField(TEXT("features"), EJson::Object))
+	{
+		UE_LOG(LogRigMapperDefinition, Warning, TEXT("Missing features field"))
+		bDefinitionValidated = false;
+		return false;
+	}
+	for (const auto& FeatureInfo : JsonObject->GetObjectField(TEXT("features"))->Values)
+	{
+		if (Inputs.Contains(FeatureInfo.Key))
+		{
+			UE_LOG(LogRigMapperDefinition, Warning, TEXT("Feature conflicting with input of similar name: %s"), *FeatureInfo.Key)
+			bDefinitionValidated = false;
+			continue;
+		}
+
+		if (!Features.AddFromJsonObject(FeatureInfo.Key, FeatureInfo.Value->AsObject()))
+		{
+			bDefinitionValidated = false;
+			continue;
+		}
+
+		FeatureNames.Add(FeatureInfo.Key);
+	}
+	return bDefinitionValidated;
+}
+
+bool URigMapperDefinition::LoadOutputsFromJsonObject(const TSharedPtr<FJsonObject>& JsonObject, const TArray<FString>& FeatureNames)
+{
+	if (!JsonObject->HasTypedField(TEXT("outputs"), EJson::Object))
+	{
+		UE_LOG(LogRigMapperDefinition, Warning, TEXT("Missing outputs field"))
+		bDefinitionValidated = false;
+		return false;
+	}
+	for (const auto& OutputAttr : JsonObject->GetObjectField(TEXT("outputs"))->Values)
+	{
+		if (OutputAttr.Key.IsEmpty())
+		{
+			UE_LOG(LogRigMapperDefinition, Warning, TEXT("Invalid output with empty name"))
+			continue;
+		}
+		FString OutputValue;
+		if (!OutputAttr.Value.IsValid() || OutputAttr.Value->IsNull() || OutputAttr.Value->AsString().IsEmpty())
+		{
+			UE_LOG(LogRigMapperDefinition, Warning, TEXT("Invalid value for output %s"), *OutputAttr.Key)
+			bDefinitionValidated = false;
+		}
+		else
+		{
+			OutputValue = OutputAttr.Value->AsString();
+			if (!Inputs.Contains(OutputValue) && !FeatureNames.Contains(OutputValue))
+			{
+				UE_LOG(LogRigMapperDefinition, Warning, TEXT("Could not find corresponding input/feature for output %s (%s)"), *OutputAttr.Key, *OutputValue)
+				bDefinitionValidated = false;
+			}
+		}
+		Outputs.Add(OutputAttr.Key, OutputValue);
+	}
+	if (Outputs.IsEmpty())
+	{
+		UE_LOG(LogRigMapperDefinition, Warning, TEXT("Not enough outputs"))
+		bDefinitionValidated = false;
+	}
+	return bDefinitionValidated;
+}
+
+#if WITH_EDITOR
+EDataValidationResult URigMapperDefinition::IsDataValid(FDataValidationContext& Context) const
+{
+	return IsDefinitionValid() ? EDataValidationResult::Valid : EDataValidationResult::Invalid;
+}
+#endif
+
+bool URigMapperDefinition::ExportAsJsonString(FString& OutJsonString) const
+{
+	if (!IsDefinitionValid())
+	{
+		return false;
+	}
+	
+	const TSharedRef<TJsonWriter<>> JsonWriter = TJsonWriterFactory<TCHAR, TPrettyJsonPrintPolicy<TCHAR>>::Create(&OutJsonString);
+
+	JsonWriter->WriteObjectStart();
+	JsonWriter->WriteValue(TEXT("inputs"), Inputs);
+	JsonWriter->WriteObjectStart(TEXT("features"));
+	for (const FRigMapperMultiplyFeature& Feature: Features.Multiply)
+	{
+		JsonWriter->WriteObjectStart(Feature.Name);
+		JsonWriter->WriteValue(TEXT("type"), TEXT("multiply"));
+		JsonWriter->WriteValue(TEXT("input_features"), Feature.Inputs);
+		JsonWriter->WriteArrayStart(TEXT("input_params"));
+		JsonWriter->WriteArrayEnd();
+		JsonWriter->WriteObjectStart(TEXT("params"));
+		JsonWriter->WriteObjectEnd();
+		JsonWriter->WriteObjectEnd();
+	}
+	for (const FRigMapperWsFeature& Feature: Features.WeightedSums)
+	{
+		JsonWriter->WriteObjectStart(Feature.Name);
+		JsonWriter->WriteValue(TEXT("type"), TEXT("weighted_sum"));
+		TArray<FString> InputNames;
+		TArray<double> InputWeights;
+		Feature.Inputs.GetKeys(InputNames);
+		InputWeights.Reserve(InputNames.Num());
+		for (const FString& Input : InputNames)
+		{
+			InputWeights.Add(Feature.Inputs[Input]);
+		}
+		JsonWriter->WriteValue(TEXT("input_features"), InputNames);
+		JsonWriter->WriteArrayStart(TEXT("input_params"));
+		JsonWriter->WriteArrayEnd();
+		JsonWriter->WriteObjectStart(TEXT("params"));
+		JsonWriter->WriteValue(TEXT("weights"), InputWeights);
+		if (Feature.Range.bHasLowerBound)
+		{
+			JsonWriter->WriteValue(TEXT("min"), Feature.Range.LowerBound);
+		}
+		if (Feature.Range.bHasUpperBound)
+		{
+			JsonWriter->WriteValue(TEXT("max"), Feature.Range.UpperBound);
+		}
+		JsonWriter->WriteObjectEnd();
+		JsonWriter->WriteObjectEnd();
+	}
+	for (const FRigMapperSdkFeature& Feature: Features.SDKs)
+	{
+		JsonWriter->WriteObjectStart(Feature.Name);
+		JsonWriter->WriteValue(TEXT("type"), TEXT("sdk"));
+		JsonWriter->WriteValue(TEXT("input_features"), TArray<FString>({Feature.Input}));
+		JsonWriter->WriteArrayStart(TEXT("input_params"));
+		JsonWriter->WriteArrayEnd();
+		JsonWriter->WriteObjectStart(TEXT("params"));
+		TArray<double> Keys;
+		TArray<double> Values;
+		Keys.Reserve(Feature.Keys.Num());
+		Values.Reserve(Feature.Keys.Num());
+		for (const FRigMapperSdkKey& Key : Feature.Keys)
+		{
+			Keys.Add(Key.In);
+			Values.Add(Key.Out);
+		}
+		JsonWriter->WriteValue(TEXT("in_val"), Keys);
+		JsonWriter->WriteValue(TEXT("out_val"), Values);	
+		JsonWriter->WriteObjectEnd();
+		JsonWriter->WriteObjectEnd();
+	}
+	JsonWriter->WriteObjectEnd();
+	JsonWriter->WriteObjectStart(TEXT("parameters"));
+	JsonWriter->WriteObjectEnd();
+	JsonWriter->WriteObjectStart(TEXT("outputs"));
+	for (const TPair<FString, FString>& Output : Outputs)
+	{
+		JsonWriter->WriteValue(Output.Key, Output.Value);	
+	}
+	JsonWriter->WriteObjectEnd();
+	JsonWriter->WriteValue(TEXT("null_outputs"), NullOutputs);
+	JsonWriter->WriteObjectEnd();
+	JsonWriter->Close();
+	return true;
+}
+
+bool URigMapperDefinition::ExportAsJsonFile(const FFilePath& JsonFilePath) const
+{
+	FString JsonString;
+	
+	if (ExportAsJsonString(JsonString))
+	{
+		return FFileHelper::SaveStringToFile(JsonString, *JsonFilePath.FilePath);
+	}
+	return false;
+}
+
+bool URigMapperDefinition::IsDefinitionValid(bool bWarn) const
+{
+	if (bDefinitionValidated)
+	{
+		return true;
+	}
+
+	bool bValid = true;
+
+	if (Inputs.IsEmpty())
+	{
+		UE_LOG(LogRigMapperDefinition, Warning, TEXT("Not enough inputs"))
+		bValid = false;
+	}
+	if (Outputs.IsEmpty())
+	{
+		UE_LOG(LogRigMapperDefinition, Warning, TEXT("Not enough outputs"))
+		bValid = false;
+	}
+
+	TArray<FString> CheckInputs;
+	CheckInputs.Reserve(Inputs.Num());
+	for (const FString& Input : Inputs)
+	{
+		// todo: check with skeleton and or control rig ?
+		if (CheckInputs.Contains(Input))
+		{
+			UE_LOG(LogRigMapperDefinition, Warning, TEXT("Duplicate input %s"), *Input)
+			bValid = false;
+		}
+		else
+		{
+			CheckInputs.Add(Input);
+		}
+	}
+	TArray<FString> FeatureNames;
+	Features.GetFeatureNames(FeatureNames);
+	for (const TPair<FString, FString>& Output : Outputs)
+	{
+		// todo: check with skeleton and or control rig ?
+		if (!Inputs.Contains(Output.Value) && !FeatureNames.Contains(Output.Value))
+		{
+			UE_LOG(LogRigMapperDefinition, Warning, TEXT("Output %s does not link to any existing input or feature"), *Output.Key)
+			bValid = false;
+		}
+	}
+	
+	bValid &= Features.IsValid(Inputs, bWarn);	
+
+	// check that NullOutputs does not contain any keys from Outputs, or any duplicated keys
+	TArray<FString> CheckNullOutputs;
+	CheckNullOutputs.Reserve(NullOutputs.Num());
+
+	for (const FString& NullOutput : NullOutputs)
+	{
+		if (Outputs.Contains(NullOutput))
+		{
+			UE_LOG(LogRigMapperDefinition, Warning, TEXT("Output is also defined as a NullOutput %s"), *NullOutput)
+			bValid = false;
+		}
+		if (CheckNullOutputs.Contains(NullOutput))
+		{
+			UE_LOG(LogRigMapperDefinition, Warning, TEXT("Duplicate NullOutput %s"), *NullOutput)
+			bValid = false;
+		}
+		else
+		{
+			CheckNullOutputs.Add(NullOutput);
+		}
+	}
+
+	return bValid;
+}
+
+bool URigMapperDefinition::Validate()
+{
+	UE_LOG(LogRigMapperDefinition, Log, TEXT("Validating definition %s"), *GetName())
+	
+	bDefinitionValidated = false;
+	bDefinitionValidated = IsDefinitionValid(true);
+
+	UE_LOG(LogRigMapperDefinition, Log, TEXT("Definition %s is %s"), *GetName(), bDefinitionValidated ? TEXT("valid") : TEXT("invalid"))
+	
+	return bDefinitionValidated;
+}
+
+#if WITH_EDITOR
+void URigMapperDefinition::PostEditChangeProperty(FPropertyChangedEvent& PropertyChangedEvent)
+{
+	Super::PostEditChangeProperty(PropertyChangedEvent);
+	OnRigMapperDefinitionUpdated.Broadcast();
+
+	bDefinitionValidated = false;
+}
+#endif
+
+bool URigMapperLinkedDefinitions::GetBakedInputRec(const FString& InputName, const int32 DefinitionIndex, FRigMapperFeature::FBakedInput& OutBakedInput, bool& bOutMissingIsNotNullOutput)
+{
+	check(SourceDefinitions[DefinitionIndex])
+	
+	if (SourceDefinitions[DefinitionIndex]->Inputs.Contains(InputName))
+	{
+		if (DefinitionIndex > 0)
+		{
+			if (!SourceDefinitions[DefinitionIndex - 1]->Outputs.Contains(InputName))
+			{
+				if (!SourceDefinitions[DefinitionIndex - 1]->NullOutputs.Contains(InputName))
+				{
+					UE_LOG(LogRigMapperDefinition, Warning, TEXT("Input %s on definition %s (layer %d) does not match any output from definition %s (layer %d)"), *InputName, *SourceDefinitions[DefinitionIndex]->GetName(), DefinitionIndex, *SourceDefinitions[DefinitionIndex - 1]->GetName(), DefinitionIndex - 1)
+					bOutMissingIsNotNullOutput = true;
+				}
+				return false;
+			}
+			return GetBakedInputRec(SourceDefinitions[DefinitionIndex - 1]->Outputs[InputName], DefinitionIndex - 1, OutBakedInput, bOutMissingIsNotNullOutput);
+		}
+		
+		OutBakedInput.Key = MakeShared<FRigMapperFeature>(InputName);
+		OutBakedInput.Value.Empty();
+		return OutBakedInput.Key.IsValid();
+	}
+
+	ERigMapperFeatureType FeatureType;
+	if (FRigMapperFeature* Feature = SourceDefinitions[DefinitionIndex]->Features.Find(InputName, FeatureType))
+	{
+		return Feature->BakeInput(this, DefinitionIndex, OutBakedInput);
+	}
+	UE_LOG(LogRigMapperDefinition, Warning, TEXT("Could not bake input %s"), *InputName)
+	return false;
+}
+
+bool FRigMapperMultiplyFeature::BakeInput(URigMapperLinkedDefinitions* LinkedDefinitions, int32 DefinitionIndex, FBakedInput& OutBakedInput) const
+{
+	// todo: refactor all bakeinput using baseclass
+	const FString BakedInputName = FString::Printf(TEXT("%s:%d"), *Name, DefinitionIndex);
+	
+	const TSharedPtr<FRigMapperMultiplyFeature> BakedMultFeature = MakeShared<FRigMapperMultiplyFeature>(BakedInputName);
+	OutBakedInput.Key = BakedMultFeature;
+	OutBakedInput.Value.Empty();
+	bool bMissingIsNotNullOutput = false;
+	
+	for (const FString& FeatureInput : Inputs)
+	{
+		FBakedInput SubFeatureBakedInput;
+		if (!LinkedDefinitions->GetBakedInputRec(FeatureInput, DefinitionIndex, SubFeatureBakedInput, bMissingIsNotNullOutput))
+		{
+			continue;
+		}
+
+		if (SubFeatureBakedInput.Key->GetFeatureType() == ERigMapperFeatureType::Multiply)
+		{
+			const FRigMapperMultiplyFeature* SubFeatureMultBakedInput = static_cast<FRigMapperMultiplyFeature*>(SubFeatureBakedInput.Key.Get());
+			BakedMultFeature->Inputs.Append(SubFeatureMultBakedInput->Inputs);
+		}
+		else
+		{
+			BakedMultFeature->Inputs.Add(SubFeatureBakedInput.Key->Name);
+			OutBakedInput.Value.Add(SubFeatureBakedInput.Key);
+		}
+		OutBakedInput.Value.Append(SubFeatureBakedInput.Value);
+	}
+
+	return OutBakedInput.Key.IsValid() && !bMissingIsNotNullOutput; // if we have a missing definition which is not defined as a NullOutput, flag a failure
+}
+
+bool FRigMapperWsFeature::BakeInput(URigMapperLinkedDefinitions* LinkedDefinitions, int32 DefinitionIndex, FBakedInput& OutBakedInput) const
+{
+	const FString BakedInputName = FString::Printf(TEXT("%s:%d"), *Name, DefinitionIndex);
+	
+	const TSharedPtr<FRigMapperWsFeature> BakedWsFeature = MakeShared<FRigMapperWsFeature>(BakedInputName);
+	OutBakedInput.Key = BakedWsFeature;
+	bool bMissingIsNotNullOutput = false;
+			
+	for (const TPair<FString, double>& FeatureInput : Inputs)
+	{
+		FBakedInput SubFeatureBakedInput;
+		if (!LinkedDefinitions->GetBakedInputRec(FeatureInput.Key, DefinitionIndex, SubFeatureBakedInput, bMissingIsNotNullOutput))
+		{
+			continue;
+		}
+
+		if (SubFeatureBakedInput.Key->GetFeatureType() == ERigMapperFeatureType::WeightedSum)
+		{
+			const FRigMapperWsFeature* SubFeatureWsBakedInput = static_cast<FRigMapperWsFeature*>(SubFeatureBakedInput.Key.Get());
+
+			for (const TPair<FString, double>& SubFeatureSubInput : SubFeatureWsBakedInput->Inputs)
+			{
+				BakedWsFeature->Inputs.Add(SubFeatureSubInput.Key, SubFeatureSubInput.Value * FeatureInput.Value);
+			}
+		}
+		else
+		{
+			BakedWsFeature->Inputs.Add(SubFeatureBakedInput.Key->Name, FeatureInput.Value);
+			OutBakedInput.Value.Add(SubFeatureBakedInput.Key);
+		}
+		OutBakedInput.Value.Append(SubFeatureBakedInput.Value);
+	}
+
+	return OutBakedInput.Key.IsValid() && !bMissingIsNotNullOutput; // if we have a missing definition which is not defined as a NullOutput, flag a failure
+}
+
+bool FRigMapperSdkFeature::BakeInput(URigMapperLinkedDefinitions* LinkedDefinitions, int32 DefinitionIndex, FBakedInput& OutBakedInput) const
+{
+	const FString BakedInputName = FString::Printf(TEXT("%s:%d"), *Name, DefinitionIndex);
+	
+	const TSharedPtr<FRigMapperSdkFeature> BakedSdkFeature = MakeShared<FRigMapperSdkFeature>(BakedInputName);
+	OutBakedInput.Key = BakedSdkFeature;
+
+	FBakedInput SubFeatureBakedInput;
+	bool bMissingIsNotNullOutput = false;
+
+	if (!LinkedDefinitions->GetBakedInputRec(Input, DefinitionIndex, SubFeatureBakedInput, bMissingIsNotNullOutput))
+	{
+		return false;
+	}
+
+	if (SubFeatureBakedInput.Key->GetFeatureType() == ERigMapperFeatureType::SDK)
+	{
+		const FRigMapperSdkFeature* SubFeatureSdkBakedInput = static_cast<FRigMapperSdkFeature*>(SubFeatureBakedInput.Key.Get());
+		BakedSdkFeature->Input = SubFeatureSdkBakedInput->Input;
+		BakeKeys(*SubFeatureSdkBakedInput, *this, BakedSdkFeature->Keys);
+	}
+	else
+	{
+		BakedSdkFeature->Input = SubFeatureBakedInput.Key->Name;
+		BakedSdkFeature->Keys = Keys;
+		OutBakedInput.Value.Add(SubFeatureBakedInput.Key);
+	}
+	OutBakedInput.Value.Append(SubFeatureBakedInput.Value);
+
+	return OutBakedInput.Key.IsValid() && !bMissingIsNotNullOutput; // if we have a missing definition which is not defined as a NullOutput, flag a failure
+}
+
+bool FRigMapperSdkFeature::BakeKeys(const FRigMapperSdkFeature& InSdk, const FRigMapperSdkFeature& OutSdk, TArray<FRigMapperSdkKey>& BakedKeys)
+{
+	// Create an array to backward evaluate InSdk (First Layer)
+	TArray<TPair<double,double>> InKeysForEval_Backward;
+	InKeysForEval_Backward.Reserve(InSdk.Keys.Num());
+	for (const FRigMapperSdkKey& Key : InSdk.Keys)
+	{
+		InKeysForEval_Backward.Add({ Key.Out, Key.In });
+	}
+	InKeysForEval_Backward.Sort([](const TPair<double,double>& A, const TPair<double,double>& B) { return A.Key < B.Key; } );
+
+	// Create an array to evaluate OutSdk (Second Layer)
+	TArray<TPair<double,double>> OutKeysForEval;
+	OutKeysForEval.Reserve(OutSdk.Keys.Num());
+	for (const FRigMapperSdkKey& Key : OutSdk.Keys)
+	{
+		OutKeysForEval.Add({ Key.In, Key.Out });
+	}
+
+	// Prepare to store our baked keys 
+	BakedKeys.Reserve(FMath::Max(InSdk.Keys.Num(), OutSdk.Keys.Num()));
+
+	// Bake all keys from InSdk (layer 1) using OutSdk
+	for (const FRigMapperSdkKey& InKey : InSdk.Keys)
+	{
+		double Value = 0.f;
+
+		if (!FacialRigMapping::FEvalNodePiecewiseLinear::Evaluate_Static(InKey.Out, OutKeysForEval, Value))
+		{
+			return false;
+		}
+		BakedKeys.Add({ InKey.In, Value });
+	}
+
+	// Because we might miss the precision from some keys of OutSdk (i.e. if we have 3 keys in OutSdk and 2 in InSdk), we now need to insert missing keys from OutSdk by reverse baking them using InSdk 
+	for (const FRigMapperSdkKey& OutKey : OutSdk.Keys)
+	{
+		double ActualInValue = 0.f;
+				
+		if (!FacialRigMapping::FEvalNodePiecewiseLinear::Evaluate_Static(OutKey.In, InKeysForEval_Backward, ActualInValue))
+		{
+			return false;
+		}
+		
+		for (int32 InKeyIndex = 0; InKeyIndex < InSdk.Keys.Num(); InKeyIndex++)
+		{
+			const bool bInsertBefore = ActualInValue < InSdk.Keys[InKeyIndex].In && (InKeyIndex == 0 || ActualInValue > InSdk.Keys[InKeyIndex - 1].In);
+			const bool bInsertAfter = ActualInValue > InSdk.Keys[InKeyIndex].In && (InKeyIndex == InSdk.Keys.Num() - 1 || ActualInValue < InSdk.Keys[InKeyIndex + 1].In);
+			
+			if (bInsertBefore || bInsertAfter)
+			{
+				const int32 NewIndex = InKeyIndex + BakedKeys.Num() - InSdk.Keys.Num(); 
+				BakedKeys.Insert({ ActualInValue, OutKey.Out }, bInsertBefore ? NewIndex : NewIndex + 1);
+				break;
+			}
+		}
+	}
+
+	// Finally, strip any duplicate or incorrectly ordered In keys from the beginning and end.
+	// todo: make sure the rig mapper allows duplicate Input keys in the middle (i.e -1;0;0;1 -> -100;-50;50;100 should be a valid use case)
+	while (BakedKeys.Num() >= 2 && BakedKeys[0].In >= BakedKeys[1].In)
+	{
+		BakedKeys.RemoveAt(0);
+	}
+	while (BakedKeys.Num() >= 2 && BakedKeys[BakedKeys.Num() - 1].In <= BakedKeys[BakedKeys.Num() - 2].In)
+	{
+		BakedKeys.RemoveAt(BakedKeys.Num() - 1);
+	}
+	return true;
+}
+
+bool URigMapperLinkedDefinitions::AddBakedInputFeature(const TSharedPtr<FRigMapperFeature>& Feature) const
+{
+	if (Feature->GetFeatureType() == ERigMapperFeatureType::Multiply)
+	{
+		BakedDefinition->Features.Multiply.Add(*static_cast<FRigMapperMultiplyFeature*>(Feature.Get()));
+	}
+	else if (Feature->GetFeatureType() == ERigMapperFeatureType::WeightedSum)
+	{
+		BakedDefinition->Features.WeightedSums.Add(*static_cast<FRigMapperWsFeature*>(Feature.Get()));
+	}
+	else if (Feature->GetFeatureType() == ERigMapperFeatureType::SDK)
+	{
+		BakedDefinition->Features.SDKs.Add(*static_cast<FRigMapperSdkFeature*>(Feature.Get()));
+	}
+	else
+	{
+		if (!SourceDefinitions[0]->Inputs.Contains(Feature->Name))
+		{
+			UE_LOG(LogRigMapperDefinition, Warning, TEXT("Baked input could not be found within the lower level inputs"))
+			return false;
+		}
+		BakedDefinition->Inputs.AddUnique(Feature->Name);
+	}
+	return true;
+}
+
+void URigMapperLinkedDefinitions::AddBakedInputs(const TArray<FRigMapperFeature::FBakedInput>& BakedInputs, const TArray<TPair<FString, FString>>& PairedOutputs) const
+{
+	TArray<FString> FeatureNames;
+	
+	BakedDefinition->Empty();
+	
+	for (int32 OutputIndex = 0; OutputIndex < PairedOutputs.Num(); OutputIndex++)
+	{
+		BakedDefinition->Outputs.Add(PairedOutputs[OutputIndex].Key, BakedInputs[OutputIndex].Key->Name);
+
+		if (!FeatureNames.Contains(BakedInputs[OutputIndex].Key->Name) && AddBakedInputFeature(BakedInputs[OutputIndex].Key))
+		{
+			FeatureNames.Add(BakedInputs[OutputIndex].Key->Name);
+			
+			for (const TSharedPtr<FRigMapperFeature>& SubFeature : BakedInputs[OutputIndex].Value)
+			{
+				if (!FeatureNames.Contains(SubFeature->Name) && !AddBakedInputFeature(SubFeature))
+				{
+					break;
+				}
+				FeatureNames.Add(SubFeature->Name);
+			}
+		}
+	}
+}
+
+#if WITH_EDITOR
+EDataValidationResult URigMapperLinkedDefinitions::IsDataValid(FDataValidationContext& Context) const
+{
+	return AreLinkedDefinitionsValid() ? EDataValidationResult::Valid : EDataValidationResult::Invalid;
+}
+#endif
+
+TArray<FRigMapperFeature::FBakedInput> URigMapperLinkedDefinitions::GetBakedInputs(const TArray<TPair<FString, FString>>& PairedOutputs)
+{
+	TArray<FRigMapperFeature::FBakedInput> BakedInputs;
+	BakedInputs.Reserve(PairedOutputs.Num());
+	bool bMissingIsNotNullOutput = false;
+
+	for (const TPair<FString, FString>& Out : PairedOutputs)
+	{
+		FRigMapperFeature::FBakedInput BakedInput;
+		
+		if (GetBakedInputRec(Out.Value, SourceDefinitions.Num() - 1, BakedInput, bMissingIsNotNullOutput))
+		{
+			BakedInputs.Add(BakedInput);
+		}
+		else
+		{
+			UE_LOG(LogRigMapperDefinition, Warning, TEXT("Could not bake input %s associated to output %s"), *Out.Value, *Out.Key)
+		}
+	}
+
+	// todo: cleanup baked unique names (prefixes)?
+
+	return BakedInputs;
+}
+
+bool URigMapperLinkedDefinitions::BakeDefinitions()
+{
+	UE_LOG(LogRigMapperDefinition, Log, TEXT("Baking linked definition %s"), *GetName())
+	
+	if (!BakedDefinition)
+	{
+		UE_LOG(LogRigMapperDefinition, Warning, TEXT("Baked definition is unset"))
+		return false;
+	}
+	if (SourceDefinitions.Num() < 2 || !SourceDefinitions[0] || !SourceDefinitions.Last())
+	{
+		UE_LOG(LogRigMapperDefinition, Warning, TEXT("Baking requires a minimum of 2 source definitions"))
+		return false;
+	}
+	for (int32 DefIndex = 0; DefIndex < SourceDefinitions.Num(); DefIndex++)
+	{
+		if (!SourceDefinitions[DefIndex]|| !SourceDefinitions[DefIndex]->IsDefinitionValid(true))
+		{
+			UE_LOG(LogRigMapperDefinition, Warning, TEXT("Invalid source definition at index %d. Make sure to revalidate the asset if necessary"), DefIndex)
+			return false;
+		}
+	}
+	
+	const TArray<TPair<FString, FString>>& PairedOutputs = SourceDefinitions.Last()->Outputs.Array();
+	const TArray<FRigMapperFeature::FBakedInput> BakedInputs = GetBakedInputs(PairedOutputs);
+	if (BakedInputs.Num() != PairedOutputs.Num())
+	{
+		UE_LOG(LogRigMapperDefinition, Warning, TEXT("The number of baked inputs do not match the expected number of outputs"))
+		return false;
+	}
+	AddBakedInputs(BakedInputs, PairedOutputs);
+
+	UE_LOG(LogRigMapperDefinition, Log, TEXT("Finished baking linked definition %s"), *GetName())
+	
+	return BakedDefinition->Validate();
+}
+
+bool URigMapperLinkedDefinitions::Validate()
+{
+	if (BakedDefinition)
+	{
+		BakedDefinition->Validate();
+	}
+
+	for (URigMapperDefinition* Definition : SourceDefinitions)
+	{
+		if (Definition)
+		{
+			Definition->Validate();
+		}
+	}
+
+	return AreLinkedDefinitionsValid();
+}
+
+bool URigMapperLinkedDefinitions::AreLinkedDefinitionsValid() const
+{
+	bool bOk = IsValid(BakedDefinition) && BakedDefinition->IsDefinitionValid();
+	
+	if (!bOk)
+	{
+		UE_LOG(LogRigMapperDefinition, Warning, TEXT("Failed to validate the baked definition"))
+	}
+
+	for (int32 DefinitionIndex = 0; DefinitionIndex < SourceDefinitions.Num(); DefinitionIndex++)
+	{
+		if (!SourceDefinitions[DefinitionIndex])
+		{
+			UE_LOG(LogRigMapperDefinition, Warning, TEXT("Source definition %d is unset"), DefinitionIndex)
+			bOk = false;
+			continue;
+		}
+		if (!SourceDefinitions[DefinitionIndex]->IsDefinitionValid())
+		{
+			UE_LOG(LogRigMapperDefinition, Warning, TEXT("Source definition %d is invalid"), DefinitionIndex)
+			bOk = false;
+		}
+		if (DefinitionIndex > 0)
+		{
+			for (const FString& Input : SourceDefinitions[DefinitionIndex]->Inputs)
+			{
+				if (!SourceDefinitions[DefinitionIndex - 1]->Outputs.Contains(Input) && !SourceDefinitions[DefinitionIndex - 1]->NullOutputs.Contains(Input))
+				{
+					UE_LOG(LogRigMapperDefinition, Warning, TEXT("Could not find matching output in definition %d for input %s in definition %d"), DefinitionIndex - 1, *Input, DefinitionIndex)
+					bOk = false;
+				}
+			}
+			for (const TPair<FString, FString>& Output : SourceDefinitions[DefinitionIndex - 1]->Outputs)
+			{
+				if (!SourceDefinitions[DefinitionIndex]->Inputs.Contains(Output.Key))
+				{
+					UE_LOG(LogRigMapperDefinition, Warning, TEXT("Could not find matching input in definition %d for output %s in definition %d"), DefinitionIndex, *Output.Key, DefinitionIndex - 1)
+				}
+			}
+		}
+	}
+	
+	return bOk;
+}
\ No newline at end of file
diff --git a/Engine/Plugins/Experimental/Animation/RigMapper/Source/RigMapper/Private/RigMapperModule.cpp b/Engine/Plugins/Experimental/Animation/RigMapper/Source/RigMapper/Private/RigMapperModule.cpp
new file mode 100644
--- /dev/null
+++ b/Engine/Plugins/Experimental/Animation/RigMapper/Source/RigMapper/Private/RigMapperModule.cpp
@@ -0,0 +1,15 @@
+// Copyright Epic Games, Inc. All Rights Reserved.
+
+#include "RigMapperModule.h"
+
+#include "Modules/ModuleManager.h"
+
+void FRigMapperModule::StartupModule()
+{
+}
+
+void FRigMapperModule::ShutdownModule()
+{
+}
+
+IMPLEMENT_MODULE(FRigMapperModule, RigMapper)
diff --git a/Engine/Plugins/Experimental/Animation/RigMapper/Source/RigMapper/Private/Tests/RigMapperTests.cpp b/Engine/Plugins/Experimental/Animation/RigMapper/Source/RigMapper/Private/Tests/RigMapperTests.cpp
new file mode 100644
--- /dev/null
+++ b/Engine/Plugins/Experimental/Animation/RigMapper/Source/RigMapper/Private/Tests/RigMapperTests.cpp
@@ -0,0 +1,690 @@
+// Copyright Epic Games, Inc. All Rights Reserved.
+
+#include "Misc/AutomationTest.h"
+#include "RigMapperDefinition.h"
+#include "RigMapper.h"
+#include "Misc/Paths.h"
+#include "UObject/UObjectGlobals.h"
+
+#if WITH_DEV_AUTOMATION_TESTS
+
+DEFINE_LOG_CATEGORY_STATIC(LogRigMapperTest, Verbose, All)
+
+
+TArray<URigMapperDefinition*> CreateValidRigMapperDefinitions1()
+{
+	TArray<URigMapperDefinition*> Definitions;
+	Definitions.SetNumUninitialized(2);
+
+	// create the first set of feature definitions, two of each type
+	FRigMapperFeatureDefinitions Features1;
+
+	TMap<FString, double> Inputs1;
+	Inputs1.Add(FString(TEXT("InputVal1")), 0.25f);
+	Inputs1.Add(FString(TEXT("InputVal2")), 0.25f);
+	Inputs1.Add(FString(TEXT("InputVal3")), 0.5f);
+	FRigMapperWsFeature WSFeature1(FString(TEXT("TestWSFeature1")));
+	WSFeature1.Inputs = Inputs1;
+	TMap<FString, double> Inputs2;
+	Inputs2.Add(FString(TEXT("InputVal4")), 0.4f);
+	Inputs2.Add(FString(TEXT("InputVal5")), 0.6f);
+	FRigMapperWsFeature WSFeature2(FString(TEXT("TestWSFeature2")));
+	WSFeature2.Inputs = Inputs2;
+	Features1.WeightedSums.Add(WSFeature1);
+	Features1.WeightedSums.Add(WSFeature2);
+
+	FRigMapperMultiplyFeature MultFeature1(FString(TEXT("TestMultFeature1")));
+	MultFeature1.Inputs = { FString(TEXT("InputVal1")), FString(TEXT("InputVal2")) };
+	FRigMapperMultiplyFeature MultFeature2(FString(TEXT("TestMultFeature2")));
+	MultFeature2.Inputs = { FString(TEXT("InputVal2")), FString(TEXT("InputVal3")) };
+	Features1.Multiply.Add(MultFeature1);
+	Features1.Multiply.Add(MultFeature2);
+
+	FRigMapperSdkFeature SDKFeature1(FString(TEXT("TestSDKFeature1")));
+	SDKFeature1.Input = { FString(TEXT("InputVal2")) };
+	SDKFeature1.Keys = { { 0, 0 }, { 0.5, 0.6 }, { 1.0, 1.0 } };
+	FRigMapperSdkFeature SDKFeature2(FString(TEXT("TestSDKFeature2")));
+	SDKFeature2.Input = { FString(TEXT("InputVal5")) };
+	SDKFeature2.Keys = { { 0, 0.25 }, { 1, 0.5 } };
+	Features1.SDKs.Add(SDKFeature1);
+	Features1.SDKs.Add(SDKFeature2);
+
+	Definitions[0] = NewObject<URigMapperDefinition>((UObject*)GetTransientPackage(), NAME_None, RF_Transient);
+	if (Definitions[0] != nullptr)
+	{
+		Definitions[0]->Inputs = { FString{ TEXT("InputVal1")}, FString{ TEXT("InputVal2")}, FString{ TEXT("InputVal3")}, FString{ TEXT("InputVal4")}, FString{ TEXT("InputVal5")} };
+		Definitions[0]->Features = Features1;
+		TMap<FString, FString> Outputs;
+		Outputs.Add(FString(TEXT("OutputVal1")), FString(TEXT("TestWSFeature1")));
+		Outputs.Add(FString(TEXT("OutputVal2")), FString(TEXT("TestWSFeature2")));
+		Outputs.Add(FString(TEXT("OutputVal3")), FString(TEXT("TestMultFeature1")));
+		Outputs.Add(FString(TEXT("OutputVal4")), FString(TEXT("TestMultFeature2")));
+		Outputs.Add(FString(TEXT("OutputVal5")), FString(TEXT("TestSDKFeature1")));
+		Outputs.Add(FString(TEXT("OutputVal6")), FString(TEXT("TestSDKFeature2")));
+		Definitions[0]->Outputs = Outputs;
+		Definitions[0]->NullOutputs = { FString{ TEXT("OutputVal7")}, FString{ TEXT("OutputVal8")} };
+	}
+
+	// create the second set of feature definitions, one of each type
+	FRigMapperFeatureDefinitions Features2;
+
+	TMap<FString, double> Inputs3;
+	Inputs3.Add(FString(TEXT("OutputVal1")), 0.2f);
+	Inputs3.Add(FString(TEXT("OutputVal2")), 0.7f);
+	Inputs3.Add(FString(TEXT("OutputVal3")), 0.1f);
+	FRigMapperWsFeature WSFeature3(FString(TEXT("TestWSFeature3")));
+	WSFeature3.Inputs = Inputs3;
+	Features2.WeightedSums.Add(WSFeature3);
+
+	FRigMapperMultiplyFeature MultFeature3(FString(TEXT("TestMultFeature3")));
+	MultFeature3.Inputs = { FString(TEXT("OutputVal4")), FString(TEXT("OutputVal5")) };
+	Features2.Multiply.Add(MultFeature3);
+
+	FRigMapperSdkFeature SDKFeature3(FString(TEXT("TestSDKFeature3")));
+	SDKFeature3.Input = { FString(TEXT("OutputVal6")) };
+	SDKFeature3.Keys = { { 0, 0 }, { 1.0, 0.8 } };
+	Features2.SDKs.Add(SDKFeature3);
+
+	Definitions[1] = NewObject<URigMapperDefinition>((UObject*)GetTransientPackage(), NAME_None, RF_Transient);
+	if (Definitions[1] != nullptr)
+	{
+		Definitions[1]->Inputs = { FString{ TEXT("OutputVal1")}, FString{ TEXT("OutputVal2")}, FString{ TEXT("OutputVal3")},
+			FString{ TEXT("OutputVal4")}, FString{ TEXT("OutputVal5")}, FString{ TEXT("OutputVal6")} };
+		Definitions[1]->Features = Features2;
+		TMap<FString, FString> Outputs;
+		Outputs.Add(FString(TEXT("OutputVal9")), FString(TEXT("TestWSFeature3")));
+		Outputs.Add(FString(TEXT("OutputVal10")), FString(TEXT("TestMultFeature3")));
+		Outputs.Add(FString(TEXT("OutputVal11")), FString(TEXT("TestSDKFeature3")));
+		Definitions[1]->Outputs = Outputs;
+		Definitions[1]->NullOutputs = { };
+	}
+
+	return Definitions;
+}
+
+TArray<URigMapperDefinition*> CreateValidRigMapperDefinitions2()
+{
+	TArray<URigMapperDefinition*> Definitions = CreateValidRigMapperDefinitions1();
+
+	// modify the second definition so it uses an Null Output from the 1st stage
+	Definitions[1]->Inputs = { FString{ TEXT("OutputVal1")}, FString{ TEXT("OutputVal2")}, FString{ TEXT("OutputVal8")}, // OutputVal8 is a Null Output from the previous stage
+		FString{ TEXT("OutputVal4")}, FString{ TEXT("OutputVal5")}, FString{ TEXT("OutputVal6")} };
+	TMap<FString, double> Inputs;
+	Inputs.Add(FString(TEXT("OutputVal1")), 0.2f);
+	Inputs.Add(FString(TEXT("OutputVal2")), 0.7f);
+	Inputs.Add(FString(TEXT("OutputVal8")), 0.1f);
+	Definitions[1]->Features.WeightedSums[0].Inputs = Inputs;
+
+	return Definitions;
+}
+
+TArray<URigMapperDefinition*> CreateInvalidRigMapperDefinitions1()
+{
+	TArray<URigMapperDefinition*> Definitions = CreateValidRigMapperDefinitions2();
+
+	// remove the NullOutputs from the first definition which makes the baked combination of the two layers invalid
+	Definitions[0]->NullOutputs = {};
+
+	return Definitions;
+}
+
+TArray<URigMapperDefinition*> CreateInvalidRigMapperDefinitions2()
+{
+	TArray<URigMapperDefinition*> Definitions = CreateValidRigMapperDefinitions1();
+
+	// remove one of the inputs for the first definition so it is referenced by a feature but no longer exists
+	Definitions[0]->Inputs = { FString{ TEXT("InputVal1")}, FString{ TEXT("InputVal2")}, FString{ TEXT("InputVal3")}, FString{ TEXT("InputVal4")}}; // removed InputVal5
+
+	// in the second definition, add a NullOutput which has the same name as one of the actual outputs
+	Definitions[1]->NullOutputs = { FString{ TEXT("OutputVal11")} };
+
+	return Definitions;
+}
+
+
+
+
+IMPLEMENT_SIMPLE_AUTOMATION_TEST(FRigMapperDefinitionTest, "RigMapper.RigMapperDefinition", EAutomationTestFlags::EditorContext | EAutomationTestFlags::EngineFilter)
+
+bool FRigMapperDefinitionTest::RunTest(const FString& InParameters)
+{	
+	TArray<URigMapperDefinition*> Definitions = CreateValidRigMapperDefinitions1();
+	URigMapperDefinition* RigMapperDefinition1 = Definitions[0];
+
+	if (!RigMapperDefinition1)
+	{
+		UE_LOG(LogRigMapperTest, Error, TEXT("Failed to create RigMapperDefinition1 asset"));
+		return false;
+	}
+
+	// check that asset is valid
+	if (!RigMapperDefinition1->IsDefinitionValid())
+	{
+		UE_LOG(LogRigMapperTest, Error, TEXT("RigMapperDefinition1 should be valid"));
+		return false;
+	}
+
+	// also check Validate method
+	if (!RigMapperDefinition1->Validate())
+	{
+		UE_LOG(LogRigMapperTest, Error, TEXT("RigMapperDefinition1 should be valid"));
+		return false;
+	}
+
+	// test that we can save and re-load the asset from a json string
+	FString JsonString1;
+	bool bExported = RigMapperDefinition1->ExportAsJsonString(JsonString1);
+	if (!bExported)
+	{
+		UE_LOG(LogRigMapperTest, Error, TEXT("Failed to export RigMapperDefinition1 asset as a Json String"));
+		return false;
+	}
+
+	// load back in to a second definition
+	URigMapperDefinition* RigMapperDefinition2 = NewObject<URigMapperDefinition>((UObject*)GetTransientPackage(), NAME_None, RF_Transient);
+	bool bImported = RigMapperDefinition2->LoadFromJsonString(JsonString1);
+	if (!bImported)
+	{
+		UE_LOG(LogRigMapperTest, Error, TEXT("Failed to import RigMapperDefinition asset from a Json String"));
+		return false;
+	}
+	FString JsonString2;
+	bExported = RigMapperDefinition2->ExportAsJsonString(JsonString2);
+	if (!bExported)
+	{
+		UE_LOG(LogRigMapperTest, Error, TEXT("Failed to export RigMapperDefinition2 asset as a Json String"));
+		return false;
+	}
+	
+	// check the same json string and also that functionally correct
+	if (JsonString1 != JsonString2)
+	{
+		UE_LOG(LogRigMapperTest, Error, TEXT("RigMapperDefinition Json export does not give the same result when re-imported"));
+		return false;
+	}
+
+	if (RigMapperDefinition2->Inputs != RigMapperDefinition1->Inputs)
+	{
+		UE_LOG(LogRigMapperTest, Error, TEXT("RigMapperDefinition Json export does not give the same result when re-imported (Inputs)"));
+		return false;
+	}
+
+	if (RigMapperDefinition2->NullOutputs != RigMapperDefinition1->NullOutputs)
+	{
+		UE_LOG(LogRigMapperTest, Error, TEXT("RigMapperDefinition Json export does not give the same result when re-imported (NullOutputs)"));
+		return false;
+	}
+
+	TArray<FString> OutputsKeys1, OutputsKeys2;
+	RigMapperDefinition1->Outputs.GetKeys(OutputsKeys1);
+	RigMapperDefinition2->Outputs.GetKeys(OutputsKeys2);
+	if (OutputsKeys1 != OutputsKeys2)
+	{
+		UE_LOG(LogRigMapperTest, Error, TEXT("RigMapperDefinition Json export does not give the same result when re-imported (Outputs)"));
+		return false;
+	}
+	for (int32 Key = 0; Key < OutputsKeys1.Num(); ++Key)
+	{
+		if (RigMapperDefinition1->Outputs[OutputsKeys1[Key]] != RigMapperDefinition2->Outputs[OutputsKeys1[Key]])
+		{
+			UE_LOG(LogRigMapperTest, Error, TEXT("RigMapperDefinition Json export does not give the same result when re-imported (Outputs)"));
+			return false;
+		}
+	}
+	
+	if (RigMapperDefinition1->Features != RigMapperDefinition2->Features)
+	{
+		UE_LOG(LogRigMapperTest, Error, TEXT("RigMapperDefinition Json export does not give the same result when re-imported (Features)"));
+		return false;
+	}
+
+	// repeat the test above but from file
+	FString Path = FPaths::ProjectSavedDir() / TEXT("test_export.json");
+	bExported = RigMapperDefinition1->ExportAsJsonFile({ Path });
+	if (!bExported)
+	{
+		UE_LOG(LogRigMapperTest, Error, TEXT("Failed to export RigMapperDefinition1 asset as a Json file"));
+		return false;
+	}
+
+	URigMapperDefinition* RigMapperDefinition3 = NewObject<URigMapperDefinition>((UObject*)GetTransientPackage(), NAME_None, RF_Transient);
+	bImported = RigMapperDefinition3->LoadFromJsonString(JsonString1);
+	if (!bImported)
+	{
+		UE_LOG(LogRigMapperTest, Error, TEXT("Failed to import RigMapperDefinition asset from file"));
+		return false;
+	}
+	FString JsonString3;
+	bExported = RigMapperDefinition3->ExportAsJsonString(JsonString3);
+	if (!bExported || JsonString3 != JsonString1 )
+	{
+		UE_LOG(LogRigMapperTest, Error, TEXT("File-based json export / import of RigMapperDefinition did not work"));
+		return false;
+	}
+
+	// test Empty method
+	RigMapperDefinition2->Empty();
+	if (!RigMapperDefinition2->Inputs.IsEmpty() || !RigMapperDefinition2->Outputs.IsEmpty() || !RigMapperDefinition2->Features.WeightedSums.IsEmpty() 
+		|| !RigMapperDefinition2->Features.SDKs.IsEmpty() || !RigMapperDefinition2->Features.Multiply.IsEmpty()
+		|| !RigMapperDefinition2->NullOutputs.IsEmpty())
+	{
+		UE_LOG(LogRigMapperTest, Error, TEXT("RigMapperDefinition expected to be empty and is not"));
+		return false;
+	}
+
+	// construct a couple of invalid definitions and check these are flagged as invalid
+	TArray<URigMapperDefinition*> Definitions2 = CreateInvalidRigMapperDefinitions2();
+
+	if (Definitions2[0]->IsDefinitionValid())
+	{
+		UE_LOG(LogRigMapperTest, Error, TEXT("RigMapperDefinition should be invalid"));
+		return false;
+	}
+
+	if (Definitions2[1]->IsDefinitionValid())
+	{
+		UE_LOG(LogRigMapperTest, Error, TEXT("RigMapperDefinition should be invalid"));
+		return false;
+	}
+
+	return true;
+}
+
+IMPLEMENT_SIMPLE_AUTOMATION_TEST(FRigMapperTest, "RigMapper.RigMapper", EAutomationTestFlags::EditorContext | EAutomationTestFlags::EngineFilter)
+
+using FRigMapperIndexMap = TArray<int32>;
+
+
+// the three functions below are adapted from FAnimNode_RigMapper; added here so we can test outside of the anim node
+
+FRigMapperIndexMap MakeIndexMap(const TSharedPtr<FacialRigMapping::FRigMapper>& InOutputRig, const TSharedPtr<FacialRigMapping::FRigMapper>& InInputRig)
+{
+	const TArray<FString> FromOutputs = InOutputRig->GetOutputNames();
+	const TArray<FString> ToInputs =  InInputRig->GetInputNames();
+
+	FRigMapperIndexMap IndexMapping;
+
+	IndexMapping.Reserve(ToInputs.Num());
+
+	for (const FString& Input : ToInputs)
+	{
+		IndexMapping.Add(FromOutputs.Find(Input));
+	}
+
+	return IndexMapping;
+}
+
+
+bool InitializeRigMapping(const TArray<URigMapperDefinition*> & InRigMapperDefinitions, TArray<TSharedPtr<FacialRigMapping::FRigMapper>>& OutRigMappers, 
+	TArray<FRigMapperIndexMap>& OutIndexMaps)
+{
+
+	if (InRigMapperDefinitions.IsEmpty())
+	{
+		return false;
+	}
+	for (int i = 0; i < InRigMapperDefinitions.Num(); i++)
+	{
+		if (!InRigMapperDefinitions[i])
+		{
+			UE_LOG(LogTemp, Error, TEXT("Invalid definition at index %d"), i)
+			return false;
+		}
+	}
+	OutRigMappers.Empty(InRigMapperDefinitions.Num());
+	OutIndexMaps.Empty(InRigMapperDefinitions.Num() - 1);
+
+	for (const URigMapperDefinition* Definition : InRigMapperDefinitions)
+	{
+		TSharedPtr<FacialRigMapping::FRigMapper> RigMapper = MakeShared<FacialRigMapping::FRigMapper>();
+
+		if (!RigMapper->LoadDefinition(Definition))
+		{
+			UE_LOG(LogTemp, Error, TEXT("Could not load definition %s"), *Definition->GetPathName())
+			OutRigMappers.Empty();
+			return false;
+		}
+
+		if (OutRigMappers.Num() > 0)
+		{
+			OutIndexMaps.Add(MakeIndexMap(OutRigMappers.Last(), RigMapper));
+		}
+
+		OutRigMappers.Add(RigMapper);
+	}
+
+	return true;
+}
+
+void EvaluateRigMapping(const TMap<FString, double>& InInputs, TArray<FRigMapperIndexMap>& InIndexMaps, TArray<TSharedPtr<FacialRigMapping::FRigMapper>>& OutRigMappers, TMap<FString, double> & OutOutputs)
+{
+
+	for (int i = 0; i < OutRigMappers.Num(); i++)
+	{
+		const TSharedPtr<FacialRigMapping::FRigMapper> RigMapper = OutRigMappers[i];
+		RigMapper->SetDirty();
+
+		// Set first rig mapper's input values from current pose
+		if (i == 0)
+		{
+
+			const TArray<FString>& InputNames = RigMapper->GetInputNames();
+			for (int n = 0; n < InputNames.Num(); n++)
+			{
+				RigMapper->SetDirectValue(n, InInputs[*InputNames[n]]);
+			}
+		}
+		else
+		{
+			// Other RigMappers: Link/Map prev rig mapper outputs to current inputs
+			check(InIndexMaps.IsValidIndex(i - 1))
+			const FRigMapperIndexMap& IndexMapping = InIndexMaps[i - 1];
+
+			// Get output values from previous rig mapper
+			TArray<double> PrevOutputValuesInOrder;
+			{
+				PrevOutputValuesInOrder = OutRigMappers[i - 1]->GetOutputValuesInOrder();
+			}
+
+			// Set input values for current rig mapper
+			{
+				for (int32 n = 0; n < IndexMapping.Num(); n++)
+				{
+					if (const int32 Index = IndexMapping[n]; Index != INDEX_NONE)
+					{
+						RigMapper->SetDirectValue(n, PrevOutputValuesInOrder[Index]);
+					}
+				}
+			}
+		}
+	}
+
+	// Set outputs from last rig mapper's output values 
+	if (!OutRigMappers.IsEmpty())
+	{
+		OutOutputs = OutRigMappers.Last()->GetOutputValues();
+	}
+}
+
+bool FRigMapperTest::RunTest(const FString& InParameters)
+{
+	TArray<URigMapperDefinition*> RigMapperDefinitions;
+	RigMapperDefinitions = CreateValidRigMapperDefinitions1();
+	if (RigMapperDefinitions.Num() != 2)
+	{
+		UE_LOG(LogRigMapperTest, Error, TEXT("Failed to create RigMapperDefinition assets"));
+		return false;
+	}
+
+	TArray<TSharedPtr<FacialRigMapping::FRigMapper>> RigMappers;
+	TArray<FRigMapperIndexMap> IndexMaps;
+
+	bool bInitialized = InitializeRigMapping(RigMapperDefinitions, RigMappers, IndexMaps);
+	if (!bInitialized)
+	{
+		UE_LOG(LogRigMapperTest, Error, TEXT("Failed to initialize RigMapping"));
+		return false;
+	}
+
+	// set up some sample inputs
+	TMap<FString, double> Inputs;
+	Inputs.Add(FString(TEXT("InputVal1")), 0.25);
+	Inputs.Add(FString(TEXT("InputVal2")), 0.4);
+	Inputs.Add(FString(TEXT("InputVal3")), 0.5);
+	Inputs.Add(FString(TEXT("InputVal4")), 0.6);
+	Inputs.Add(FString(TEXT("InputVal5")), 0.75);
+
+	// evaluate the rig mappers; this evaluates the whole thing in the same way in which this is done in the anim node
+	TMap<FString, double> Outputs;
+	EvaluateRigMapping(Inputs, IndexMaps, RigMappers, Outputs);
+
+	// check that the outputs are correct
+
+	// Inputs
+	//	InputVal1 = 0.25
+	//	InputVal2 = 0.4
+	//	InputVal3 = 0.5
+	//	InputVal4 = 0.6
+	//	InputVal5 = 0.75
+	// 
+	// Layer1
+	//	TestWSFeature1 = InputVal1 * 0.25 + InputVal2 * 0.25 + InputVal3 * 0.5 = 0.25 * 0.25 + 0.4 * 0.25 + 0.5 * 0.5 = 0.4125
+	//	TestWSFeature2 = InputVal4 * 0.4 + InputVal5 * 0.6 = 0.6 * 0.4 + 0.75 * 0.6 = 0.69
+	//	TestMultFeature1 = InputVal1 * InputVal2 = 0.25 * 0.4 = 0.1
+	//	TestMultFeature2 = InputVal2 * InputVal3 = 0.4 * 0.5 = 0.2
+	//	TestSDKFeature1 = InputVal2 SDK values (0,0) (0.5,0.6) (1.0,1.0) = 0.48
+	//	TestSDKFeature2 = InputVal5 SDK values (0,0.25) (1,0.5) = 0.4375
+	// 
+	//	Outputs:
+	// 		OutputVal1 = TestWSFeature1 = 0.4125
+	//		OutputVal2 = TestWSFeature2 = 0.69
+	//		OutputVal3 = TestMultFeature1 = 0.1
+	//		OutputVal4 = TestMultFeature2 = 0.2
+	//		OutputVal5 = TestSDKFeature1 = 0.48
+	//		OutputVal6 = TestSDKFeature2 = 0.4375
+	// 
+	// Layer2
+	//	TestWSFeature3 = OutputVal1 * 0.2 + OutputVal2 * 0.7 + OutputVal3 * 0.1 = 0.4125 * 0.2 + 0.69 * 0.7 + 0.1 * 0.1 = 0.0925
+	//	TestMultFeature3 = OutputVal4 * OutputVal5 = 0.096
+	//	TestSDKFeature3 = OutputVal6 SDK values (0,0) (1.0, 0.8) = 0.35
+	//
+	//	Outputs:
+	//		OutputVal9 = TestWSFeature3 = 0.5755
+	//		OutputVal10 = TestMultFeature3 = 0.096
+	//		OutputVal11 = TestSDKFeature3 = 0.35
+
+	const double* OutputVal9 = Outputs.Find(FString(TEXT("OutputVal9")));
+	const double* OutputVal10 = Outputs.Find(FString(TEXT("OutputVal10")));
+	const double* OutputVal11 = Outputs.Find(FString(TEXT("OutputVal11")));
+
+	if (!OutputVal9 || !OutputVal10 || !OutputVal11 || !FMath::IsNearlyEqual(*OutputVal9, 0.5755, SMALL_NUMBER)
+		|| !FMath::IsNearlyEqual(*OutputVal10, 0.096, SMALL_NUMBER) || !FMath::IsNearlyEqual(*OutputVal11, 0.35, SMALL_NUMBER))
+	{
+		UE_LOG(LogRigMapperTest, Error, TEXT("Output values for tested rig-mapping are not as expected"));
+		return false;
+	}
+
+	// TODO would also be good to test this in the anim node but beyond the scope of the current tests
+	return true;
+}
+
+IMPLEMENT_SIMPLE_AUTOMATION_TEST(RigMapperLinkedDefinitionsTest, "RigMapper.RigMapperLinkedDefinitions", EAutomationTestFlags::EditorContext | EAutomationTestFlags::EngineFilter)
+
+bool RigMapperLinkedDefinitionsTest::RunTest(const FString& InParameters)
+{
+	// create two rigmapper definitions
+	TArray<URigMapperDefinition*> RigMapperDefinitions1 = CreateValidRigMapperDefinitions1();
+	if (RigMapperDefinitions1.Num() != 2)
+	{
+		UE_LOG(LogRigMapperTest, Error, TEXT("Failed to create RigMapperDefinition assets"));
+		return false;
+	}
+
+	// link the definitions and bake
+	URigMapperLinkedDefinitions * LinkedDefinitions1 = NewObject<URigMapperLinkedDefinitions>((UObject*)GetTransientPackage(), NAME_None, RF_Transient);
+	if (!LinkedDefinitions1)
+	{
+		UE_LOG(LogRigMapperTest, Error, TEXT("Failed to create a RigMapperLinkedDefinitions"));
+		return false;
+	}
+
+	LinkedDefinitions1->BakedDefinition = NewObject<URigMapperDefinition>((UObject*)GetTransientPackage(), NAME_None, RF_Transient);
+	if (!LinkedDefinitions1)
+	{
+		UE_LOG(LogRigMapperTest, Error, TEXT("Failed to create RigMapperDefinition"));
+		return false;
+	}
+
+	LinkedDefinitions1->SourceDefinitions = RigMapperDefinitions1;
+	bool bBaked = LinkedDefinitions1->BakeDefinitions();
+
+	if (!bBaked)
+	{
+		UE_LOG(LogRigMapperTest, Error, TEXT("Failed to bake definitions"));
+		return false;
+	}
+
+	// now evaluate the result of the baked definitions and check that the results are as expected
+	TArray<URigMapperDefinition*> BakedRigMapperDefinitions;
+	BakedRigMapperDefinitions.SetNumUninitialized(1);
+	BakedRigMapperDefinitions[0] = LinkedDefinitions1->BakedDefinition;
+
+	TArray<TSharedPtr<FacialRigMapping::FRigMapper>> RigMappers;
+	TArray<FRigMapperIndexMap> IndexMaps;
+
+	bool bInitialized = InitializeRigMapping(BakedRigMapperDefinitions, RigMappers, IndexMaps);
+	if (!bInitialized)
+	{
+		UE_LOG(LogRigMapperTest, Error, TEXT("Failed to initialize RigMapping"));
+		return false;
+	}
+
+	// set up some sample inputs
+	TMap<FString, double> Inputs;
+	Inputs.Add(FString(TEXT("InputVal1")), 0.25);
+	Inputs.Add(FString(TEXT("InputVal2")), 0.4);
+	Inputs.Add(FString(TEXT("InputVal3")), 0.5);
+	Inputs.Add(FString(TEXT("InputVal4")), 0.6);
+	Inputs.Add(FString(TEXT("InputVal5")), 0.75);
+
+	// evaluate the rig mappers; this evaluates the whole thing in the same way in which this is done in the anim node
+	TMap<FString, double> Outputs;
+	EvaluateRigMapping(Inputs, IndexMaps, RigMappers, Outputs);
+
+
+	const double* OutputVal9 = Outputs.Find(FString(TEXT("OutputVal9")));
+	const double* OutputVal10 = Outputs.Find(FString(TEXT("OutputVal10")));
+	const double* OutputVal11 = Outputs.Find(FString(TEXT("OutputVal11")));
+
+	if (!OutputVal9 || !OutputVal10 || !OutputVal11 || !FMath::IsNearlyEqual(*OutputVal9, 0.5755, SMALL_NUMBER)
+		|| !FMath::IsNearlyEqual(*OutputVal10, 0.096, SMALL_NUMBER) || !FMath::IsNearlyEqual(*OutputVal11, 0.35, SMALL_NUMBER))
+	{
+		UE_LOG(LogRigMapperTest, Error, TEXT("Output values for tested baked rig-mapping are not as expected"));
+		return false;
+	}
+
+	// test other methods
+	if (!LinkedDefinitions1->AreLinkedDefinitionsValid())
+	{
+		UE_LOG(LogRigMapperTest, Error, TEXT("LinkedDefinitions are expected to be valid, and are not"));
+		return false;
+	}
+
+	if (!LinkedDefinitions1->Validate())
+	{
+		UE_LOG(LogRigMapperTest, Error, TEXT("LinkedDefinitions are expected to be valid, and are not"));
+		return false;
+	}
+
+	const TArray<TPair<FString, FString>>& PairedOutputs = RigMapperDefinitions1.Last()->Outputs.Array();
+	const TArray<FRigMapperFeature::FBakedInput> BakedInputs = LinkedDefinitions1->GetBakedInputs(PairedOutputs);
+	if (BakedInputs.Num() != PairedOutputs.Num())
+	{
+		UE_LOG(LogRigMapperDefinition, Error, TEXT("The number of baked inputs do not match the expected number of outputs"))
+			return false;
+	}
+
+	// no explicit test for GetBakedInputRec here
+
+	// test for a case where one of the NullOutputs is used; this should still be valid
+	TArray<URigMapperDefinition*> RigMapperDefinitions2 = CreateValidRigMapperDefinitions2();
+
+	URigMapperLinkedDefinitions* LinkedDefinitions2 = NewObject<URigMapperLinkedDefinitions>((UObject*)GetTransientPackage(), NAME_None, RF_Transient);
+	if (!LinkedDefinitions2)
+	{
+		UE_LOG(LogRigMapperTest, Error, TEXT("Failed to create a RigMapperLinkedDefinitions"));
+		return false;
+	}
+
+	LinkedDefinitions2->BakedDefinition = NewObject<URigMapperDefinition>((UObject*)GetTransientPackage(), NAME_None, RF_Transient);
+	if (!LinkedDefinitions2)
+	{
+		UE_LOG(LogRigMapperTest, Error, TEXT("Failed to create RigMapperDefinition"));
+		return false;
+	}
+
+	LinkedDefinitions2->SourceDefinitions = RigMapperDefinitions2;
+	bBaked = LinkedDefinitions2->BakeDefinitions();
+
+	if (!bBaked)
+	{
+		UE_LOG(LogRigMapperTest, Error, TEXT("Failed to bake definitions"));
+		return false;
+	}
+
+	if (!LinkedDefinitions2->AreLinkedDefinitionsValid())
+	{
+		UE_LOG(LogRigMapperTest, Error, TEXT("LinkedDefinitions are expected to be valid, and are not"));
+		return false;
+	}
+
+	// now test for a couple of invalid cases and check that validation catches these
+
+	// case 1 is definitions which are valid individually but are missing an input to the second set of definitions
+	TArray<URigMapperDefinition*> RigMapperDefinitions3 = CreateInvalidRigMapperDefinitions1();
+
+	URigMapperLinkedDefinitions* LinkedDefinitions3 = NewObject<URigMapperLinkedDefinitions>((UObject*)GetTransientPackage(), NAME_None, RF_Transient);
+	if (!LinkedDefinitions3)
+	{
+		UE_LOG(LogRigMapperTest, Error, TEXT("Failed to create a RigMapperLinkedDefinitions"));
+		return false;
+	}
+
+	LinkedDefinitions3->BakedDefinition = NewObject<URigMapperDefinition>((UObject*)GetTransientPackage(), NAME_None, RF_Transient);
+	if (!LinkedDefinitions3)
+	{
+		UE_LOG(LogRigMapperTest, Error, TEXT("Failed to create RigMapperDefinition"));
+		return false;
+	}
+
+	LinkedDefinitions3->SourceDefinitions = RigMapperDefinitions3;
+	bBaked = LinkedDefinitions3->BakeDefinitions();
+
+	if (bBaked)
+	{
+		UE_LOG(LogRigMapperTest, Error, TEXT("Successfully baked invalid definitions 3 which is not expected"));
+		return false;
+	}
+
+	if (LinkedDefinitions3->AreLinkedDefinitionsValid())
+	{
+		UE_LOG(LogRigMapperTest, Error, TEXT("LinkedDefinitions 3 are expected to be invalid, and are actually valid"));
+		return false;
+	}
+
+	// case 2 is simply an example where the individual definitions are invalid
+	TArray<URigMapperDefinition*> RigMapperDefinitions4 = CreateInvalidRigMapperDefinitions1();
+
+	URigMapperLinkedDefinitions* LinkedDefinitions4 = NewObject<URigMapperLinkedDefinitions>((UObject*)GetTransientPackage(), NAME_None, RF_Transient);
+	if (!LinkedDefinitions4)
+	{
+		UE_LOG(LogRigMapperTest, Error, TEXT("Failed to create a RigMapperLinkedDefinitions"));
+		return false;
+	}
+
+	LinkedDefinitions4->BakedDefinition = NewObject<URigMapperDefinition>((UObject*)GetTransientPackage(), NAME_None, RF_Transient);
+	if (!LinkedDefinitions4)
+	{
+		UE_LOG(LogRigMapperTest, Error, TEXT("Failed to create RigMapperDefinition"));
+		return false;
+	}
+
+	LinkedDefinitions4->SourceDefinitions = RigMapperDefinitions4;
+	bBaked = LinkedDefinitions4->BakeDefinitions();
+
+	if (bBaked)
+	{
+		UE_LOG(LogRigMapperTest, Error, TEXT("Successfully baked invalid definitions 4 which is not expected"));
+		return false;
+	}
+
+	if (LinkedDefinitions4->AreLinkedDefinitionsValid())
+	{
+		UE_LOG(LogRigMapperTest, Error, TEXT("LinkedDefinitions 4 are expected to be invalid, and are actually valid"));
+		return false;
+	}
+
+	return true;
+}
+
+#endif
\ No newline at end of file
diff --git a/Engine/Plugins/Experimental/Animation/RigMapper/Source/RigMapper/Public/AnimNode_RigMapper.h b/Engine/Plugins/Experimental/Animation/RigMapper/Source/RigMapper/Public/AnimNode_RigMapper.h
new file mode 100644
--- /dev/null
+++ b/Engine/Plugins/Experimental/Animation/RigMapper/Source/RigMapper/Public/AnimNode_RigMapper.h
@@ -0,0 +1,82 @@
+// Copyright Epic Games, Inc. All Rights Reserved.
+
+#pragma once
+
+#include "Animation/AnimNodeBase.h"
+
+#include "CoreMinimal.h"
+#include "RigMapper.h"
+
+#include "AnimNode_RigMapper.generated.h"
+
+class URigMapperDefinitionUserData;
+class URigMapperDefinition;
+struct FRigControlElement;
+struct FRigCurveElement;
+class UControlRig;
+class USkeletalMesh;
+
+DECLARE_STATS_GROUP(TEXT("RigMapper"), STATGROUP_RigMapper, STATCAT_Advanced);
+
+DECLARE_CYCLE_STAT_EXTERN(TEXT("Evaluate AnimNode"), STAT_RigMapperEvaluateNode, STATGROUP_RigMapper, RIGMAPPER_API);
+DECLARE_CYCLE_STAT_EXTERN(TEXT("Evaluate RigMapping"), STAT_RigMapperEvaluateAllRigs, STATGROUP_RigMapper, RIGMAPPER_API);
+DECLARE_CYCLE_STAT_EXTERN(TEXT("Evaluate RigMapper"), STAT_RigMapperEvaluateRig, STATGROUP_RigMapper, RIGMAPPER_API);
+DECLARE_CYCLE_STAT_EXTERN(TEXT("Update RigMapper"), STAT_RigMapperUpdateValues, STATGROUP_RigMapper, RIGMAPPER_API);
+
+using FRigMapperIndexMap = TArray<int32>;
+
+/**
+ * 
+ */
+USTRUCT(BlueprintInternalUseOnly)
+struct RIGMAPPER_API FAnimNode_RigMapper : public FAnimNode_Base
+{
+	GENERATED_USTRUCT_BODY()
+	
+public:
+	FAnimNode_RigMapper();
+	virtual ~FAnimNode_RigMapper() override;
+	
+	// FAnimNode_Base interface
+	virtual void OnInitializeAnimInstance(const FAnimInstanceProxy* InProxy, const UAnimInstance* InAnimInstance) override;
+	virtual void Initialize_AnyThread(const FAnimationInitializeContext& Context) override;
+	virtual void GatherDebugData(FNodeDebugData& DebugData) override;
+	virtual bool HasPreUpdate() const override { return true; }
+	virtual void PreUpdate(const UAnimInstance* InAnimInstance) override;
+	virtual void Update_AnyThread(const FAnimationUpdateContext& Context) override;
+	virtual void CacheBones_AnyThread(const FAnimationCacheBonesContext& Context) override;
+	virtual void Evaluate_AnyThread(FPoseContext & Output) override;
+	void EvaluateRigMapping(FPoseContext& Output) const;
+	virtual int32 GetLODThreshold() const override;
+	// End of FAnimNode_Base interface
+
+	bool InitializeRigMapping(USkeletalMesh* TargetMesh=nullptr);
+
+public:
+	UPROPERTY(EditAnywhere, BlueprintReadWrite, Category = Links)
+	FPoseLink SourcePose;
+
+	UPROPERTY(EditAnywhere, BlueprintReadWrite, Category = RigMapper, meta = (PinHiddenByDefault))
+	TArray<TObjectPtr<URigMapperDefinition>> Definitions;
+	
+	/*
+	* Max LOD that this node is allowed to run
+	* For example if you have LODThreshold to be 2, it will run until LOD 2 (based on 0 index)
+	* when the component LOD becomes 3, it will stop update/evaluate
+	* currently transition would be issue and that has to be re-visited
+	*/
+	UPROPERTY(EditAnywhere, BlueprintReadWrite, Category = Performance, meta = (DisplayName = "LOD Threshold"))
+	int32 LODThreshold = INDEX_NONE;
+
+	UPROPERTY(EditAnywhere, BlueprintReadWrite, Category = Alpha, meta = (PinShownByDefault))
+	float Alpha = 1.0f;
+
+private:	
+	static FRigMapperIndexMap MakeIndexMap(const TSharedPtr<FacialRigMapping::FRigMapper>& OutputRig, const TSharedPtr<FacialRigMapping::FRigMapper>& InputRig);
+	
+private:
+	TArray<TSharedPtr<FacialRigMapping::FRigMapper>> RigMappers;
+	TArray<TObjectPtr<URigMapperDefinition>> LoadedDefinitions;
+	TObjectPtr<URigMapperDefinitionUserData> LoadedUserData;
+	TArray<FRigMapperIndexMap> IndexMaps;
+};
diff --git a/Engine/Plugins/Experimental/Animation/RigMapper/Source/RigMapper/Public/RigMapper.h b/Engine/Plugins/Experimental/Animation/RigMapper/Source/RigMapper/Public/RigMapper.h
new file mode 100644
--- /dev/null
+++ b/Engine/Plugins/Experimental/Animation/RigMapper/Source/RigMapper/Public/RigMapper.h
@@ -0,0 +1,144 @@
+// Copyright Epic Games, Inc. All Rights Reserved.
+
+#pragma once
+
+#include "CoreMinimal.h"
+
+struct FRigMapperMultiplyFeature;
+struct FRigMapperWsFeature;
+struct FRigMapperSdkFeature;
+class URigMapperDefinition;
+
+namespace FacialRigMapping
+{
+	class FEvalNode;
+	using EvalPtr = TSharedPtr<FEvalNode>;
+
+	/**
+	 * 
+	 */
+	class RIGMAPPER_API FRigMapper
+	{
+	public:
+		bool LoadDefinition(const URigMapperDefinition* Definition);
+		bool SetDirectValue(int32 Index, double Value);
+		bool SetDirectValue(const FString& InputName, double Value);
+		void Reset();
+		bool IsValid() const;
+
+		TMap<FString,double> GetOutputValues(bool bSkipUnset = false) const;
+		const TArray<FString>& GetOutputNames() const;
+		TArray<double> GetOutputValuesInOrder() const;
+		
+		void SetDirty() const;
+
+		const TArray<FString>& GetInputNames() const;
+
+	private:
+		TMap<FString, EvalPtr> Nodes;
+		
+		TArray<EvalPtr> OutputNodes;
+		TArray<FString> OutputNames;
+		
+		TArray<EvalPtr> InputNodes;
+		TArray<FString> InputNames;
+	};
+
+	/**
+	 *
+	 */
+	class FEvalNode
+	{
+	public:
+		explicit FEvalNode(const FString& InName);
+		virtual ~FEvalNode(){}
+	
+		double GetValue() const;
+		bool TryGetValue(double &OutValue) const;
+		void SetDirect(double Value);
+		void Reset();
+
+		virtual void Initialize(const TMap<FString, EvalPtr>& Nodes) {}
+		
+		bool IsInitialized() const;
+
+	protected:
+		bool bInitialized = false;
+		
+	private:
+		virtual bool Evaluate(double &OutValue) const { return false; }
+
+		FString Name;
+		mutable TOptional<double> CachedValue;
+	};
+
+	/**
+	 *
+	 */
+	class FEvalInputNode : public FEvalNode
+	{
+	public:
+		explicit FEvalInputNode(const FString& InName) : FEvalNode(InName) {}
+		
+		virtual void Initialize(const TMap<FString, EvalPtr>& Nodes) override { bInitialized = true; }
+	};
+
+	/**
+	 * 
+	 */
+	class FEvalNodeWeightedSum : public FEvalNode
+	{
+	public:
+		FEvalNodeWeightedSum(const FRigMapperWsFeature& FeatureDefinition);
+
+	private:
+		virtual bool Evaluate(double &OutValue) const override;
+
+		virtual void Initialize(const TMap<FString, EvalPtr>& Nodes) override;
+
+		TMap<FString, double> WeightedLinkedInputNames;
+	
+		TMap<EvalPtr, double> WeightedLinkedInputs;
+
+		TRange<double> Range;
+	};
+
+	/**
+	 *
+	 */
+	class FEvalNodePiecewiseLinear :public FEvalNode
+	{
+	public:
+		FEvalNodePiecewiseLinear(const FRigMapperSdkFeature& FeatureDefinition);
+
+	public:
+		static bool Evaluate_Static(const double InputValue, const TArray<TPair<double,double>>& Keys, double &OutValue);
+		
+	private:
+		virtual bool Evaluate(double &OutValue) const override;
+		
+		virtual void Initialize(const TMap<FString, EvalPtr>& Nodes) override;
+
+	private:
+		FString LinkedInputName;
+		EvalPtr LinkedInput;
+		TArray<TPair<double,double>> Keys;
+	};
+
+	/**
+	 * 
+	 */
+	class FEvalNodeMultiply : public FEvalNode
+	{
+	public:
+		explicit FEvalNodeMultiply(const FRigMapperMultiplyFeature& FeatureDefinition);
+
+	private:
+		virtual bool Evaluate(double &OutValue) const override;
+
+		virtual void Initialize(const TMap<FString, EvalPtr>& Nodes) override;
+
+		TArray<FString> LinkedInputNames;
+		TArray<EvalPtr> LinkedInputs;
+	};
+}
diff --git a/Engine/Plugins/Experimental/Animation/RigMapper/Source/RigMapper/Public/RigMapperDefinition.h b/Engine/Plugins/Experimental/Animation/RigMapper/Source/RigMapper/Public/RigMapperDefinition.h
new file mode 100644
--- /dev/null
+++ b/Engine/Plugins/Experimental/Animation/RigMapper/Source/RigMapper/Public/RigMapperDefinition.h
@@ -0,0 +1,372 @@
+// Copyright Epic Games, Inc. All Rights Reserved.
+
+#pragma once
+
+#include "CoreMinimal.h"
+#include "Engine/DataAsset.h"
+#include "Engine/AssetUserData.h"
+
+#include "RigMapperDefinition.generated.h"
+
+
+DECLARE_LOG_CATEGORY_EXTERN(LogRigMapperDefinition, Log, All);
+
+class FJsonObject;
+class FJsonValue;
+class URigMapperLinkedDefinitions;
+
+UENUM(BlueprintType)
+enum class ERigMapperFeatureType : uint8
+{
+	Input,
+	WeightedSum,
+	SDK,
+	Multiply
+};
+
+USTRUCT(BlueprintType)
+struct RIGMAPPER_API FRigMapperFeature
+{
+	GENERATED_BODY()
+
+public:
+	using FBakedInput = TPair<TSharedPtr<FRigMapperFeature>, TArray<TSharedPtr<FRigMapperFeature>>>;
+	
+	UPROPERTY(EditAnywhere, BlueprintReadWrite, Category="Animation|Rig Mapper", meta=(DisplayPriority=0))
+	FString Name;
+
+public:
+	virtual ~FRigMapperFeature() = default;
+	
+	FRigMapperFeature() {}
+
+	explicit FRigMapperFeature(const FString& InName) : Name(InName) {}
+
+	virtual bool LoadFromJsonObject(const TSharedPtr<FJsonObject>& JsonObject) { return false; };
+
+	virtual bool IsValid(const TArray<FString>& InputNames, bool bWarn = false) const;
+
+	virtual ERigMapperFeatureType GetFeatureType() const { return ERigMapperFeatureType::Input; }
+
+	virtual void GetInputs(TArray<FString>& OutInputs) const { }
+
+	virtual bool BakeInput(URigMapperLinkedDefinitions* LinkedDefinitions, int32 DefinitionIndex, FBakedInput& OutBakedInput) const { return false; }
+
+
+	bool operator==(const FRigMapperFeature& InOther) const
+	{
+		return Name == InOther.Name;
+	}
+	
+protected:
+	bool GetJsonArray(TSharedPtr<FJsonObject> JsonObject, TArray<TSharedPtr<FJsonValue>>& OutArray, const FString& Identifier, const FString& OwnerIdentifier="") const;
+};
+
+
+USTRUCT(BlueprintType)
+struct RIGMAPPER_API FRigMapperMultiplyFeature : public FRigMapperFeature
+{
+	GENERATED_BODY()
+
+public:
+	UPROPERTY(EditAnywhere, BlueprintReadWrite, Category="Animation|Rig Mapper", meta=(DisplayPriority=1))
+	TArray<FString> Inputs;
+	
+public:
+	FRigMapperMultiplyFeature() {}
+
+	explicit FRigMapperMultiplyFeature(const FString& InName) : FRigMapperFeature(InName) {}
+
+	virtual bool LoadFromJsonObject(const TSharedPtr<FJsonObject>& JsonObject) override;
+
+	virtual bool IsValid(const TArray<FString>& InputNames, bool bWarn = false) const override;
+
+	virtual ERigMapperFeatureType GetFeatureType() const override { return ERigMapperFeatureType::Multiply; }
+
+	virtual void GetInputs(TArray<FString>& OutInputs) const override { OutInputs = Inputs; }
+
+	virtual bool BakeInput(URigMapperLinkedDefinitions* LinkedDefinitions, int32 DefinitionIndex, FBakedInput& OutBakedInput) const override;
+
+	bool operator==(const FRigMapperMultiplyFeature& InOther) const
+	{
+		return FRigMapperFeature::operator==(InOther) && Inputs == InOther.Inputs;
+	}
+};
+
+USTRUCT(BlueprintType)
+struct RIGMAPPER_API FRigMapperFeatureRange
+{
+	GENERATED_BODY()
+
+public:
+	UPROPERTY(EditAnywhere, BlueprintReadWrite, Category="Animation|Rig Mapper")
+	bool bHasLowerBound = false;
+	
+	UPROPERTY(EditAnywhere, BlueprintReadWrite, Category="Animation|Rig Mapper")
+	double LowerBound = 0;
+
+	UPROPERTY(EditAnywhere, BlueprintReadWrite, Category="Animation|Rig Mapper")
+	bool bHasUpperBound = false;
+	
+	UPROPERTY(EditAnywhere, BlueprintReadWrite, Category="Animation|Rig Mapper")
+	double UpperBound = 0;
+
+	bool operator==(const FRigMapperFeatureRange& InOther) const
+	{
+		return bHasLowerBound == InOther.bHasLowerBound && FMath::IsNearlyEqual(LowerBound, InOther.LowerBound, SMALL_NUMBER)
+			&& bHasUpperBound == InOther.bHasUpperBound && FMath::IsNearlyEqual(UpperBound, InOther.UpperBound, SMALL_NUMBER);
+	}
+};
+
+USTRUCT(BlueprintType)
+struct RIGMAPPER_API FRigMapperWsFeature : public FRigMapperFeature
+{
+	GENERATED_BODY()
+
+public:
+	UPROPERTY(EditAnywhere, BlueprintReadWrite, Category="Animation|Rig Mapper", meta=(DisplayPriority=1))
+	TMap<FString, double> Inputs;
+
+	UPROPERTY(EditAnywhere, BlueprintReadWrite, Category="Animation|Rig Mapper", meta=(DisplayPriority=1))
+	FRigMapperFeatureRange Range;
+	
+public:
+	FRigMapperWsFeature() {}
+
+	explicit FRigMapperWsFeature(const FString& InName) : FRigMapperFeature(InName) {}
+
+	virtual bool LoadFromJsonObject(const TSharedPtr<FJsonObject>& JsonObject) override;
+
+	virtual bool IsValid(const TArray<FString>& InputNames, bool bWarn = false) const override;
+
+	virtual ERigMapperFeatureType GetFeatureType() const override { return ERigMapperFeatureType::WeightedSum; }
+
+	virtual void GetInputs(TArray<FString>& OutInputs) const override;
+
+	virtual bool BakeInput(URigMapperLinkedDefinitions* LinkedDefinitions, int32 DefinitionIndex, FBakedInput& OutBakedInput) const override;
+
+	bool operator==(const FRigMapperWsFeature& InOther) const
+	{
+		if (Inputs.Num() != InOther.Inputs.Num())
+		{
+			return false; 
+		}
+
+		for (const TPair<FString, double>& Pair : Inputs)
+		{
+			const double* OtherValue = InOther.Inputs.Find(Pair.Key);
+			if (OtherValue == nullptr || !FMath::IsNearlyEqual(Pair.Value, *OtherValue, SMALL_NUMBER))
+			{
+				return false; 
+			}
+		}
+
+		return FRigMapperFeature::operator==(InOther) && Range == InOther.Range;
+	}
+};
+
+USTRUCT(BlueprintType)
+struct RIGMAPPER_API FRigMapperSdkKey
+{
+	GENERATED_BODY()
+
+public:
+	UPROPERTY(EditAnywhere, BlueprintReadWrite, Category="Animation|Rig Mapper")
+	double In = 0;
+
+	UPROPERTY(EditAnywhere, BlueprintReadWrite, Category="Animation|Rig Mapper")
+	double Out = 0;
+
+	bool operator==(const FRigMapperSdkKey& InOther) const
+	{
+		return FMath::IsNearlyEqual(In, InOther.In, SMALL_NUMBER) && FMath::IsNearlyEqual(Out, InOther.Out, SMALL_NUMBER);
+	}
+};
+
+USTRUCT(BlueprintType)
+struct RIGMAPPER_API FRigMapperSdkFeature : public FRigMapperFeature
+{
+	GENERATED_BODY()
+	
+public:
+	UPROPERTY(EditAnywhere, BlueprintReadWrite, Category="Animation|Rig Mapper", meta=(DisplayPriority=1))
+	FString Input;
+	
+	UPROPERTY(EditAnywhere, BlueprintReadWrite, Category="Animation|Rig Mapper", meta=(DisplayPriority=1))
+	TArray<FRigMapperSdkKey> Keys;
+	
+public:
+	FRigMapperSdkFeature() {}
+
+	explicit FRigMapperSdkFeature(const FString& InName) : FRigMapperFeature(InName) {};
+
+	virtual bool LoadFromJsonObject(const TSharedPtr<FJsonObject>& JsonObject) override;
+	
+	virtual bool IsValid(const TArray<FString>& InputNames, bool bWarn = false) const override;
+
+	virtual ERigMapperFeatureType GetFeatureType() const override { return ERigMapperFeatureType::SDK; }
+	
+	virtual void GetInputs(TArray<FString>& OutInputs) const override { OutInputs = { Input }; };
+
+	virtual bool BakeInput(URigMapperLinkedDefinitions* LinkedDefinitions, int32 DefinitionIndex, FBakedInput& OutBakedInput) const override;
+
+	static bool BakeKeys(const FRigMapperSdkFeature& InSdk, const FRigMapperSdkFeature& OutSdk, TArray<FRigMapperSdkKey>& BakedKeys);
+
+	bool operator==(const FRigMapperSdkFeature& InOther) const
+	{
+		return FRigMapperFeature::operator==(InOther) && Input == InOther.Input && Keys == InOther.Keys;
+	}
+};
+
+USTRUCT(BlueprintType)
+struct RIGMAPPER_API FRigMapperFeatureDefinitions
+{
+	GENERATED_BODY()
+
+	UPROPERTY(EditAnywhere, BlueprintReadWrite, Category="Animation|Rig Mapper")
+	TArray<FRigMapperMultiplyFeature> Multiply;
+
+	UPROPERTY(EditAnywhere, BlueprintReadWrite, Category="Animation|Rig Mapper")
+	TArray<FRigMapperWsFeature> WeightedSums;
+
+	UPROPERTY(EditAnywhere, BlueprintReadWrite, Category="Animation|Rig Mapper")
+	TArray<FRigMapperSdkFeature> SDKs;
+
+public:
+	bool AddFromJsonObject(const FString& FeatureName, const TSharedPtr<FJsonObject>& JsonObject);
+
+	bool GetFeatureNames(TArray<FString>& OutFeatureNames) const;
+	
+	bool IsValid(const TArray<FString>& InputNames, bool bWarn = false) const;
+
+	void Empty();
+	
+	FRigMapperFeature* Find(const FString& FeatureName, ERigMapperFeatureType& OutFeatureType);
+
+	bool operator==(const FRigMapperFeatureDefinitions& InOther) const
+	{
+		return Multiply == InOther.Multiply && WeightedSums == InOther.WeightedSums && SDKs == InOther.SDKs;
+	}
+};
+
+
+DECLARE_MULTICAST_DELEGATE(FOnRigMapperDefinitionUpdated);
+
+UCLASS(BlueprintType)
+class RIGMAPPER_API URigMapperDefinition : public UDataAsset
+{
+	GENERATED_BODY()
+
+public:
+	// Delegate to notify listeners that we have loaded definitions
+	FOnRigMapperDefinitionUpdated OnRigMapperDefinitionUpdated;
+
+	UPROPERTY(EditAnywhere, BlueprintReadWrite, Category="Animation|Rig Mapper")
+	TArray<FString> Inputs;
+
+	UPROPERTY(EditAnywhere, BlueprintReadWrite, Category="Animation|Rig Mapper")
+	FRigMapperFeatureDefinitions Features;
+	
+	UPROPERTY(EditAnywhere, BlueprintReadWrite, Category="Animation|Rig Mapper")
+	TMap<FString, FString> Outputs;
+
+	UPROPERTY(EditAnywhere, BlueprintReadWrite, Category="Animation|Rig Mapper")
+	TArray<FString> NullOutputs;
+
+private:
+	UPROPERTY()
+	FString JsonSource;
+
+	UPROPERTY()
+	bool bDefinitionValidated = false;
+
+public:
+	UFUNCTION(BlueprintCallable, Category="Animation|Rig Mapper")
+	bool LoadFromJsonFile(const FFilePath& JsonFilePath);
+	
+	UFUNCTION(BlueprintCallable, Category="Animation|Rig Mapper")
+	bool LoadFromJsonString(const FString& JsonString);
+	
+	UFUNCTION(BlueprintCallable, Category="Animation|Rig Mapper")
+	bool ExportAsJsonString(FString& OutJsonString) const;
+
+	UFUNCTION(BlueprintCallable, Category="Animation|Rig Mapper")
+	bool ExportAsJsonFile(const FFilePath& JsonFilePath) const;
+	
+	UFUNCTION(BlueprintCallable, Category="Animation|Rig Mapper")
+	void Empty();
+	
+	UFUNCTION(BlueprintCallable, Category="Animation|Rig Mapper")
+	bool IsDefinitionValid(bool bWarn = false) const;
+	
+	UFUNCTION(BlueprintCallable, Category="Animation|Rig Mapper")
+	bool Validate();
+	
+private:
+	bool LoadInputsFromJsonObject(const TSharedPtr<FJsonObject>& JsonObject);
+	bool LoadFeaturesFromJsonObject(const TSharedPtr<FJsonObject>& JsonObject, TArray<FString>& FeatureNames);
+	bool LoadOutputsFromJsonObject(const TSharedPtr<FJsonObject>& JsonObject, const TArray<FString>& FeatureNames);
+	bool LoadNullOutputsFromJsonObject(const TSharedPtr<FJsonObject>& JsonObject);
+
+#if WITH_EDITOR
+public:
+	virtual void PostEditChangeProperty(FPropertyChangedEvent& PropertyChangedEvent) override;
+
+private:	
+	virtual EDataValidationResult IsDataValid(FDataValidationContext& Context) const override;
+#endif
+	
+	// todo: helper function to generate new asset from json file
+	// todo: invalidate on property changed
+};
+
+UCLASS(BlueprintType)
+class RIGMAPPER_API URigMapperLinkedDefinitions : public UDataAsset
+{
+	GENERATED_BODY()
+
+public:
+	UPROPERTY(EditAnywhere, BlueprintReadWrite, Category="Animation|Rig Mapper")
+	TArray<TObjectPtr<URigMapperDefinition>> SourceDefinitions;
+
+	UPROPERTY(EditAnywhere, BlueprintReadWrite, Category="Animation|Rig Mapper")
+	TObjectPtr<URigMapperDefinition> BakedDefinition;
+	
+public:
+	// todo: bake to new
+	// todo: invalidate on property or source definition changed 
+	
+	UFUNCTION(BlueprintCallable, Category="Animation|Rig Mapper")
+	bool BakeDefinitions();
+	
+	UFUNCTION(BlueprintCallable, Category="Animation|Rig Mapper")
+	bool Validate();
+
+	UFUNCTION(BlueprintCallable, Category="Animation|Rig Mapper")
+	bool AreLinkedDefinitionsValid() const;
+
+	TArray<FRigMapperFeature::FBakedInput> GetBakedInputs(const TArray<TPair<FString, FString>>& PairedOutputs);
+	
+	bool GetBakedInputRec(const FString& InputName, const int32 DefinitionIndex, FRigMapperFeature::FBakedInput& OutBakedInput, bool& bOutMissingIsNotNullOutput);
+	
+private:
+	bool AddBakedInputFeature(const TSharedPtr<FRigMapperFeature>& Feature) const;
+	
+	void AddBakedInputs(const TArray<FRigMapperFeature::FBakedInput>& BakedInputs, const TArray<TPair<FString, FString>>& PairedOutputs) const;
+	// todo: helper functions to bake to new / existing asset
+	// todo: helper function to import/export from/to json files
+
+#if WITH_EDITOR
+	virtual EDataValidationResult IsDataValid(FDataValidationContext& Context) const override;
+#endif
+};
+
+UCLASS(NotBlueprintable, HideCategories = (Object))
+class RIGMAPPER_API URigMapperDefinitionUserData : public UAssetUserData
+{
+	GENERATED_BODY()
+
+public:
+	UPROPERTY(EditAnywhere, BlueprintReadWrite, Category="Animation|Rig Mapper")
+	TArray<TObjectPtr<URigMapperDefinition>> Definitions;
+};
\ No newline at end of file
diff --git a/Engine/Plugins/Experimental/Animation/RigMapper/Source/RigMapper/Public/RigMapperModule.h b/Engine/Plugins/Experimental/Animation/RigMapper/Source/RigMapper/Public/RigMapperModule.h
new file mode 100644
--- /dev/null
+++ b/Engine/Plugins/Experimental/Animation/RigMapper/Source/RigMapper/Public/RigMapperModule.h
@@ -0,0 +1,18 @@
+// Copyright Epic Games, Inc. All Rights Reserved.
+
+#pragma once
+
+#include "Modules/ModuleInterface.h"
+
+/**
+ * FRigMapperModule
+ * The Runtime Module for the Generic Sequence plugin
+ */
+class FRigMapperModule : public IModuleInterface
+{
+public:
+	//~ Begin IModuleInterface interface
+	virtual void StartupModule() override;
+	virtual void ShutdownModule() override;
+	//~ End IModuleInterface interface
+};
diff --git a/Engine/Plugins/Experimental/Animation/RigMapper/Source/RigMapper/RigMapper.Build.cs b/Engine/Plugins/Experimental/Animation/RigMapper/Source/RigMapper/RigMapper.Build.cs
new file mode 100644
--- /dev/null
+++ b/Engine/Plugins/Experimental/Animation/RigMapper/Source/RigMapper/RigMapper.Build.cs
@@ -0,0 +1,25 @@
+// Copyright Epic Games, Inc. All Rights Reserved.
+
+using UnrealBuildTool;
+
+public class RigMapper : ModuleRules
+{
+	public RigMapper(ReadOnlyTargetRules Target) : base(Target)
+	{
+		PublicDependencyModuleNames.AddRange(
+			new string[]
+			{
+				"Core",
+				"CoreUObject",
+				"Engine",
+				"ControlRig",
+				"AnimationCore",
+				"Json"
+			});
+
+		PrivateDependencyModuleNames.AddRange(
+			new string[]
+			{
+			});
+	}
+}
diff --git a/Engine/Plugins/Experimental/Animation/RigMapper/Source/RigMapperEditor/Private/AnimGraphNode_RigMapper.cpp b/Engine/Plugins/Experimental/Animation/RigMapper/Source/RigMapperEditor/Private/AnimGraphNode_RigMapper.cpp
new file mode 100644
--- /dev/null
+++ b/Engine/Plugins/Experimental/Animation/RigMapper/Source/RigMapperEditor/Private/AnimGraphNode_RigMapper.cpp
@@ -0,0 +1,77 @@
+// Copyright Epic Games, Inc. All Rights Reserved.
+
+#include "AnimGraphNode_RigMapper.h"
+
+#define LOCTEXT_NAMESPACE "RigMapper"
+
+UAnimGraphNode_RigMapper::UAnimGraphNode_RigMapper(const FObjectInitializer& ObjectInitializer) : Super(ObjectInitializer)
+{
+}
+
+FText UAnimGraphNode_RigMapper::GetNodeTitle(ENodeTitleType::Type TitleType) const
+{
+	return LOCTEXT("RigMapper", "Rig Mapper");
+}
+
+FText UAnimGraphNode_RigMapper::GetTooltipText() const
+{
+	return LOCTEXT("AnimGraphNode_RigMapper_ToolTip", "Evaluates an output pose from the current pose using a control rig backward solve and a set of operations configured in a Json file");
+}
+
+FText UAnimGraphNode_RigMapper::GetMenuCategory() const
+{
+	return LOCTEXT("AnimGraphNode_RigMapper_Category", "Animation|Poses");
+}
+
+FLinearColor UAnimGraphNode_RigMapper::GetNodeBodyTintColor() const
+{
+	return FLinearColor(FColor::Emerald);
+}
+
+void UAnimGraphNode_RigMapper::ValidateAnimNodeDuringCompilation(USkeleton* ForSkeleton,
+	FCompilerResultsLog& MessageLog)
+{
+	// todo: validate config
+	Super::ValidateAnimNodeDuringCompilation(ForSkeleton, MessageLog);
+}
+
+FEditorModeID UAnimGraphNode_RigMapper::GetEditorMode() const
+{
+	// todo: override to set a custom editor mode
+	return Super::GetEditorMode();
+}
+
+EAnimAssetHandlerType UAnimGraphNode_RigMapper::SupportsAssetClass(const UClass* AssetClass) const
+{
+	// todo: Rig Mapper Config Asset
+	// if (AssetClass->IsChildOf(URigMapperConfig::StaticClass()))
+	// {
+	// 	return EAnimAssetHandlerType::Supported;
+	// }
+	// else
+	// {
+	// 	return EAnimAssetHandlerType::NotSupported;
+	// }
+
+	return EAnimAssetHandlerType::NotSupported;
+}
+
+void UAnimGraphNode_RigMapper::CopyNodeDataToPreviewNode(FAnimNode_Base* InPreviewNode)
+{
+	// todo: implement to transfer input/output to preview?
+}
+
+void UAnimGraphNode_RigMapper::PostEditChangeProperty(FPropertyChangedEvent& PropertyChangedEvent)
+{
+	Super::PostEditChangeProperty(PropertyChangedEvent);
+
+	const FName PropertyName = PropertyChangedEvent.GetPropertyName();
+	if (PropertyName == GET_MEMBER_NAME_CHECKED(FAnimNode_RigMapper, Definitions))
+	{
+		USkeletalMesh* TargetMesh = GetAnimBlueprint() ? (GetAnimBlueprint()->GetPreviewMesh() ? GetAnimBlueprint()->GetPreviewMesh() : GetAnimBlueprint()->TargetSkeleton->GetPreviewMesh()) : nullptr;
+		Node.InitializeRigMapping(TargetMesh);
+		ReconstructNode();
+	}
+}
+
+#undef LOCTEXT_NAMESPACE
diff --git a/Engine/Plugins/Experimental/Animation/RigMapper/Source/RigMapperEditor/Private/RigMapperDefinitionAssetTypeActions.cpp b/Engine/Plugins/Experimental/Animation/RigMapper/Source/RigMapperEditor/Private/RigMapperDefinitionAssetTypeActions.cpp
new file mode 100644
--- /dev/null
+++ b/Engine/Plugins/Experimental/Animation/RigMapper/Source/RigMapperEditor/Private/RigMapperDefinitionAssetTypeActions.cpp
@@ -0,0 +1,41 @@
+// Copyright Epic Games, Inc. All Rights Reserved.
+
+#include "RigMapperDefinitionAssetTypeActions.h"
+
+#include "RigMapperDefinition.h"
+#include "RigMapperDefinitionEditorToolkit.h"
+
+
+UClass* FRigMapperDefinitionAssetTypeActions::GetSupportedClass() const
+{
+	return URigMapperDefinition::StaticClass();
+}
+
+FText FRigMapperDefinitionAssetTypeActions::GetName() const
+{
+	return INVTEXT("Rig Mapper Definition");
+}
+
+FColor FRigMapperDefinitionAssetTypeActions::GetTypeColor() const
+{
+	return FColor::Yellow;
+}
+
+uint32 FRigMapperDefinitionAssetTypeActions::GetCategories()
+{
+	return EAssetTypeCategories::Animation;
+}
+
+void FRigMapperDefinitionAssetTypeActions::OpenAssetEditor(const TArray<UObject*>& InObjects, TSharedPtr<IToolkitHost> ToolkitHost)
+{
+	const EToolkitMode::Type Mode = ToolkitHost.IsValid() ? EToolkitMode::WorldCentric : EToolkitMode::Standalone;
+
+	for (UObject* Obj : InObjects)
+	{
+		if (URigMapperDefinition* Definition = Cast<URigMapperDefinition>(Obj))
+		{
+			const TSharedRef<FRigMapperDefinitionEditorToolkit> EditorToolkit = MakeShareable(new FRigMapperDefinitionEditorToolkit);
+			EditorToolkit->Initialize(Definition, Mode, ToolkitHost);
+		}
+	}
+}
\ No newline at end of file
diff --git a/Engine/Plugins/Experimental/Animation/RigMapper/Source/RigMapperEditor/Private/RigMapperDefinitionAssetTypeActions.h b/Engine/Plugins/Experimental/Animation/RigMapper/Source/RigMapperEditor/Private/RigMapperDefinitionAssetTypeActions.h
new file mode 100644
--- /dev/null
+++ b/Engine/Plugins/Experimental/Animation/RigMapper/Source/RigMapperEditor/Private/RigMapperDefinitionAssetTypeActions.h
@@ -0,0 +1,19 @@
+// Copyright Epic Games, Inc. All Rights Reserved.
+
+#pragma once
+ 
+#include "CoreMinimal.h"
+#include "AssetTypeActions_Base.h"
+
+/**
+ * The asset actions for the URigMapperDefinition data asset class and link to its asset editor toolkit
+ */
+class RIGMAPPEREDITOR_API FRigMapperDefinitionAssetTypeActions : public FAssetTypeActions_Base
+{
+public:
+	virtual UClass* GetSupportedClass() const override;
+	virtual FText GetName() const override;
+	virtual FColor GetTypeColor() const override;
+	virtual uint32 GetCategories() override;
+	virtual void OpenAssetEditor(const TArray<UObject*>& InObjects, TSharedPtr<IToolkitHost> EditWithinLevelEditor) override;
+};
diff --git a/Engine/Plugins/Experimental/Animation/RigMapper/Source/RigMapperEditor/Private/RigMapperDefinitionEditorToolkit.cpp b/Engine/Plugins/Experimental/Animation/RigMapper/Source/RigMapperEditor/Private/RigMapperDefinitionEditorToolkit.cpp
new file mode 100644
--- /dev/null
+++ b/Engine/Plugins/Experimental/Animation/RigMapper/Source/RigMapperEditor/Private/RigMapperDefinitionEditorToolkit.cpp
@@ -0,0 +1,338 @@
+// Copyright Epic Games, Inc. All Rights Reserved.
+
+#include "RigMapperDefinitionEditorToolkit.h"
+
+#include "DetailCategoryBuilder.h"
+#include "DetailLayoutBuilder.h"
+#include "RigMapperGraph/RigMapperDefinitionEditorGraphNode.h"
+#include "RigMapperGraph/SRigMapperDefinitionGraphEditor.h"
+#include "RigMapperEditorModule.h"
+
+#include "DetailsViewObjectFilter.h"
+#include "IStructureDetailsView.h"
+#include "PropertyEditorModule.h"
+#include "Modules/ModuleManager.h"
+#include "RigMapperGraph/SRigMapperDefinitionStructureView.h"
+#include "Widgets/Docking/SDockTab.h"
+#include "Widgets/Views/STreeView.h"
+
+#define LOCTEXT_NAMESPACE "RigMapperDefinitionEditorToolkit"
+
+const FName FRigMapperDefinitionEditorToolkit::DefinitionEditorGraphTabId("RigMapperEditor_DefinitionGraphView");
+const FName FRigMapperDefinitionEditorToolkit::DefinitionEditorStructureTabId("RigMapperEditor_DefinitionStructureView");
+const FName FRigMapperDefinitionEditorToolkit::DefinitionEditorDetailsTabId("RigMapperEditor_DefinitionDetailsView");
+
+const TMap<FName, ERigMapperNodeType> FRigMapperDefinitionEditorToolkit::PropertyNameToNodeTypeMapping = {
+	{ GET_MEMBER_NAME_CHECKED(URigMapperDefinition, Inputs), ERigMapperNodeType::Input },
+	{ GET_MEMBER_NAME_CHECKED(URigMapperDefinition, Features), ERigMapperNodeType::Invalid },
+	{ GET_MEMBER_NAME_CHECKED(FRigMapperFeatureDefinitions, Multiply), ERigMapperNodeType::Multiply },
+	{ GET_MEMBER_NAME_CHECKED(FRigMapperFeatureDefinitions, WeightedSums), ERigMapperNodeType::WeightedSum },
+	{ GET_MEMBER_NAME_CHECKED(FRigMapperFeatureDefinitions, SDKs), ERigMapperNodeType::SDK },
+	{ GET_MEMBER_NAME_CHECKED(URigMapperDefinition, Outputs), ERigMapperNodeType::Output },
+};
+
+void FRigMapperDefinitionEditorToolkit::Initialize(URigMapperDefinition* InDefinition, EToolkitMode::Type InMode, TSharedPtr<IToolkitHost> InToolkitHost)
+{
+	Definition = InDefinition;
+
+	// todo
+	// GEditor->GetEditorSubsystem<UImportSubsystem>()->OnAssetPostImport.AddSP(this, &FSimpleAssetEditor::HandleAssetPostImport);
+	// FCoreUObjectDelegates::OnObjectsReplaced.AddSP(this, &FSimpleAssetEditor::OnObjectsReplaced);
+	
+	const TSharedRef<FTabManager::FLayout> Layout = FTabManager::NewLayout("Standalone_RigMapperDefinitionEditor_Layout_v1")
+	->AddArea
+	(
+		FTabManager::NewPrimaryArea()->SetOrientation(Orient_Vertical)
+		->Split
+		(
+			FTabManager::NewSplitter()->SetOrientation(Orient_Horizontal)
+			->Split
+			(
+				FTabManager::NewStack()
+				->AddTab(DefinitionEditorStructureTabId, ETabState::OpenedTab)
+				->SetSizeCoefficient(0.2)
+			)
+			->Split
+			(
+				FTabManager::NewStack() 
+				->SetHideTabWell(true)
+				->AddTab(DefinitionEditorGraphTabId, ETabState::OpenedTab)
+				->SetSizeCoefficient(0.6)
+			)
+			->Split
+			(
+				FTabManager::NewStack() 
+				->AddTab(DefinitionEditorDetailsTabId, ETabState::OpenedTab)
+				->SetSizeCoefficient(0.2)
+			)
+		)
+	);
+
+	InitAssetEditor(
+		InMode,
+		InToolkitHost,
+		FRigMapperEditorModule::AppIdentifier,
+		Layout,
+		true /*bCreateDefaultStandaloneMenu*/,
+		true /*bCreateDefaultToolbar*/,
+		InDefinition
+	);
+
+	FRigMapperEditorModule::RegisterRigMapperDefinitionToolbarEntries();
+}
+
+void FRigMapperDefinitionEditorToolkit::RegisterTabSpawners(const TSharedRef<FTabManager>& InTabManager)
+{
+	WorkspaceMenuCategory = InTabManager->AddLocalWorkspaceMenuCategory(INVTEXT("Rig Mapper Definition Editor"));
+
+	FAssetEditorToolkit::RegisterTabSpawners(InTabManager);
+	
+	InTabManager->RegisterTabSpawner(DefinitionEditorGraphTabId, FOnSpawnTab::CreateSP(this, &FRigMapperDefinitionEditorToolkit::SpawnGraphTab))
+		.SetDisplayName(INVTEXT("Graph"))
+		.SetGroup(WorkspaceMenuCategory.ToSharedRef())
+		.SetIcon(FSlateIcon(FAppStyle::GetAppStyleSetName(), "GraphEditor.EventGraph_16x"));
+
+	InTabManager->RegisterTabSpawner(DefinitionEditorStructureTabId, FOnSpawnTab::CreateSP(this, &FRigMapperDefinitionEditorToolkit::SpawnStructureTab))
+		.SetDisplayName(INVTEXT("Structure"))
+		.SetGroup(WorkspaceMenuCategory.ToSharedRef())
+		.SetIcon(FSlateIcon(FAppStyle::GetAppStyleSetName(), "Kismet.Tabs.Palette"));
+
+	InTabManager->RegisterTabSpawner(DefinitionEditorDetailsTabId, FOnSpawnTab::CreateSP(this, &FRigMapperDefinitionEditorToolkit::SpawnDetailsTab))
+		.SetDisplayName(INVTEXT("Details"))
+		.SetGroup(WorkspaceMenuCategory.ToSharedRef())
+		.SetIcon(FSlateIcon(FAppStyle::GetAppStyleSetName(), "LevelEditor.Tabs.Details"));
+}
+
+void FRigMapperDefinitionEditorToolkit::UnregisterTabSpawners(const TSharedRef<FTabManager>& InTabManager)
+{
+	InTabManager->UnregisterTabSpawner(DefinitionEditorGraphTabId);
+	InTabManager->UnregisterTabSpawner(DefinitionEditorStructureTabId);
+	InTabManager->UnregisterTabSpawner(DefinitionEditorDetailsTabId);
+}
+
+FName FRigMapperDefinitionEditorToolkit::GetToolkitFName() const
+{
+	return "RigMapperDefinitionEditor";
+}
+
+FText FRigMapperDefinitionEditorToolkit::GetBaseToolkitName() const
+{
+	return INVTEXT("Rig Mapper Definition Editor");
+}
+
+FString FRigMapperDefinitionEditorToolkit::GetWorldCentricTabPrefix() const
+{
+	return "Rig Mapper Definition ";
+}
+
+FLinearColor FRigMapperDefinitionEditorToolkit::GetWorldCentricTabColorScale() const
+{
+	return FLinearColor( 0.0f, 0.1f, 0.2f, 0.5f );
+}
+
+TSharedRef<SDockTab> FRigMapperDefinitionEditorToolkit::SpawnGraphTab(const FSpawnTabArgs& SpawnTabArgs)
+{
+	// todo: search and replace (easy rename inputs, outputs and features)
+	// todo: error log
+	GraphEditor = SNew(SRigMapperDefinitionGraphEditor, Definition);
+
+	GraphEditor->OnSelectionChanged.BindRaw(this, &FRigMapperDefinitionEditorToolkit::HandleGraphSelectionChanged);
+
+	Definition->OnRigMapperDefinitionUpdated.AddRaw(this, &FRigMapperDefinitionEditorToolkit::HandleRigMapperDefinitionLoaded);
+	
+	TSharedRef<SDockTab> DockTab = SNew(SDockTab)
+	.TabColorScale(GetTabColorScale())
+	.Label(INVTEXT("Graph"))
+	[
+		GraphEditor.ToSharedRef()
+	];
+	
+	return DockTab;
+}
+
+void FRigMapperDefinitionEditorToolkit::HandleRigMapperDefinitionLoaded()
+{
+	if (GraphEditor)
+	{
+		GraphEditor->RebuildGraph();
+	}
+	if (StructureView)
+	{
+		StructureView->RebuildTree();
+	}
+}
+
+TSharedRef<IDetailCustomization> FRigMapperDefinitionEditorToolkit::FDetailsViewCustomization::MakeInstance()
+{
+	return MakeShared<FDetailsViewCustomization>();
+}
+
+void FRigMapperDefinitionEditorToolkit::FDetailsViewCustomization::CustomizeDetails(IDetailLayoutBuilder& DetailLayout)
+{
+	TArray<TWeakObjectPtr<UObject>> Objects;
+	DetailLayout.GetObjectsBeingCustomized(Objects);
+	
+	DetailLayout.EditCategory("Animation|Rig Mapper", INVTEXT("Elements"));
+	DetailLayout.HideCategory("Animation");
+	DetailLayout.HideCategory("Rig Mapper");
+}
+
+TSharedRef<SDockTab> FRigMapperDefinitionEditorToolkit::SpawnDetailsTab(const FSpawnTabArgs& SpawnTabArgs)
+{
+	FPropertyEditorModule& PropertyModule = FModuleManager::LoadModuleChecked<FPropertyEditorModule>("PropertyEditor");
+	
+	FDetailsViewArgs Args;
+	Args.bHideSelectionTip = true;
+	Args.NameAreaSettings = FDetailsViewArgs::HideNameArea;
+
+	DetailsView = PropertyModule.CreateDetailView(Args);
+	DetailsView->RegisterInstancedCustomPropertyLayout(URigMapperDefinition::StaticClass(), FOnGetDetailCustomizationInstance::CreateStatic(&FDetailsViewCustomization::MakeInstance));
+	DetailsView->SetIsPropertyVisibleDelegate(FIsPropertyVisible::CreateRaw(this, &FRigMapperDefinitionEditorToolkit::HandleIsPropertyVisible));
+	DetailsView->OnFinishedChangingProperties().AddRaw(this, &FRigMapperDefinitionEditorToolkit::HandleFinishedChangingProperties);
+	DetailsView->SetObject(Definition);
+	if (Definition)
+	{
+		Definition->SetFlags(RF_Transactional);
+	}
+
+	TSharedRef<SDockTab> DockTab = SNew(SDockTab)
+	.TabColorScale(GetTabColorScale())
+	.Label(INVTEXT("Details"))
+	[
+		DetailsView.ToSharedRef()
+	];
+		
+	return DockTab;
+}
+
+TSharedRef<SDockTab> FRigMapperDefinitionEditorToolkit::SpawnStructureTab(const FSpawnTabArgs& SpawnTabArgs)
+{
+	StructureView = SNew(SRigMapperDefinitionStructureView, Definition);
+	StructureView->OnSelectionChanged.BindRaw(this, &FRigMapperDefinitionEditorToolkit::HandleStructureSelectionChanged);
+		
+	TSharedRef<SDockTab> DockTab = SNew(SDockTab)
+	.TabColorScale(GetTabColorScale())
+	.Label(INVTEXT("Structure"))
+	[
+		StructureView.ToSharedRef()
+	];
+
+	return DockTab;
+}
+
+
+bool FRigMapperDefinitionEditorToolkit::HandleIsPropertyVisible(const FPropertyAndParent& PropertyAndParent)
+{
+	const int32 ArrayIndex = PropertyAndParent.ArrayIndex;
+	const FName PropertyName = PropertyAndParent.Property.GetFName();
+	
+	if (StructureView->IsSelectionEmpty())
+	{
+		return true;
+	}
+	
+	if (PropertyAndParent.ParentProperties.Num() > 2)
+	{
+		return true;
+	}
+	
+	if (PropertyNameToNodeTypeMapping.Contains(PropertyName))
+	{
+		return StructureView->IsNodeOrChildSelected(PropertyNameToNodeTypeMapping[PropertyName], ArrayIndex);
+	}
+	if (PropertyName == GET_MEMBER_NAME_CHECKED(URigMapperDefinition, NullOutputs))
+	{
+		// todo: add a way to select the null outputs from the structure
+		return false;
+	}
+	return true;
+}
+
+void FRigMapperDefinitionEditorToolkit::HandleFinishedChangingProperties(const FPropertyChangedEvent& PropertyChangedEvent)
+{
+	FPropertyChangedEvent Event = PropertyChangedEvent;
+	Event.ObjectIteratorIndex = 0;
+	
+	const EPropertyChangeType::Type ChangeType = Event.ChangeType;
+	const FName PropertyName = Event.GetPropertyName();
+	const FName ParentName = Event.GetMemberPropertyName();
+	const FName InputName = GET_MEMBER_NAME_CHECKED(URigMapperDefinition, Inputs);
+	const FName OutputName = GET_MEMBER_NAME_CHECKED(URigMapperDefinition, Outputs);
+	
+	const bool bInput = PropertyName == InputName && ParentName == InputName;
+	const bool bOutput = !bInput && (PropertyName == OutputName || PropertyName == OutputName.ToString() + TEXT("_Key")) && ParentName == OutputName; // PropertyName will be Outputs_Key if renaming an output
+	const bool bName = !bInput && !bOutput && PropertyName == GET_MEMBER_NAME_CHECKED(FRigMapperFeature, Name);
+
+	const bool bSelectNew = ChangeType == EPropertyChangeType::Duplicate || ChangeType == EPropertyChangeType::ArrayAdd;
+	const bool bRename = ChangeType == EPropertyChangeType::ValueSet && (bInput || bOutput || bName);
+	
+	GraphEditor->RebuildGraph();
+	StructureView->RebuildTree();
+	
+	if (bSelectNew || bRename)
+	{
+		if (bInput)
+		{
+			int32 InputIndex = Event.GetArrayIndex(PropertyName.ToString());
+			StructureView->SelectNode(Definition->Inputs[InputIndex], ERigMapperNodeType::Input, true);
+			GraphEditor->SelectNodes({ Definition->Inputs[InputIndex] }, {}, {});
+		}
+		else if (bOutput)
+		{
+			int32 OutputIndex = Event.GetArrayIndex(PropertyName.ToString());
+			TArray<FString> OutputsInOrder;
+			Definition->Outputs.GenerateKeyArray(OutputsInOrder);
+			StructureView->SelectNode(OutputsInOrder[OutputIndex], ERigMapperNodeType::Output, true);
+			GraphEditor->SelectNodes({}, {}, { OutputsInOrder[OutputIndex] });
+		}
+		else if (ParentName == GET_MEMBER_NAME_CHECKED(URigMapperDefinition, Features))
+		{
+			if (const int32 MultiplyIndex = Event.GetArrayIndex(GET_MEMBER_NAME_CHECKED(FRigMapperFeatureDefinitions, Multiply).ToString()); MultiplyIndex != INDEX_NONE)
+			{
+				StructureView->SelectNode(Definition->Features.Multiply[MultiplyIndex].Name, ERigMapperNodeType::Multiply, true);
+				GraphEditor->SelectNodes({}, { Definition->Features.Multiply[MultiplyIndex].Name }, {});
+			}
+			else if (const int32 WsIndex = Event.GetArrayIndex(GET_MEMBER_NAME_CHECKED(FRigMapperFeatureDefinitions, WeightedSums).ToString()); WsIndex != INDEX_NONE)
+			{
+				StructureView->SelectNode(Definition->Features.WeightedSums[WsIndex].Name, ERigMapperNodeType::WeightedSum, true);
+				GraphEditor->SelectNodes({}, { Definition->Features.WeightedSums[WsIndex].Name }, { });
+			}
+			else if (const int32 SdkIndex = Event.GetArrayIndex(GET_MEMBER_NAME_CHECKED(FRigMapperFeatureDefinitions, SDKs).ToString()); SdkIndex != INDEX_NONE)
+			{
+				StructureView->SelectNode(Definition->Features.SDKs[SdkIndex].Name, ERigMapperNodeType::SDK, true);
+				GraphEditor->SelectNodes({}, { Definition->Features.SDKs[SdkIndex].Name }, {});
+			}
+		}
+	}
+
+	// todo: restore selection
+	
+	// todo: remove from selection if even remove / clear
+}
+
+void FRigMapperDefinitionEditorToolkit::HandleGraphSelectionChanged(const TSet<UObject*>& Nodes)
+{
+	StructureView->ClearSelection();
+	
+	for (const UObject* ObjectNode : Nodes)
+	{
+		if (const URigMapperDefinitionEditorGraphNode* Node = Cast<URigMapperDefinitionEditorGraphNode>(ObjectNode))
+		{
+			StructureView->SelectNode(Node->GetNodeName(), Node->GetNodeType(), true);
+		}
+	}
+	
+	DetailsView->ForceRefresh();
+}
+
+void FRigMapperDefinitionEditorToolkit::HandleStructureSelectionChanged(ESelectInfo::Type SelectInfo, TArray<FString> SelectedInputs, TArray<FString> SelectedFeatures, TArray<FString> SelectedOutputs)
+{
+	DetailsView->ForceRefresh();
+
+	if (SelectInfo != ESelectInfo::Direct)
+	{
+		GraphEditor->SelectNodes(SelectedInputs, SelectedFeatures, SelectedOutputs);
+	};
+}
+
+#undef LOCTEXT_NAMESPACE
diff --git a/Engine/Plugins/Experimental/Animation/RigMapper/Source/RigMapperEditor/Private/RigMapperDefinitionEditorToolkit.h b/Engine/Plugins/Experimental/Animation/RigMapper/Source/RigMapperEditor/Private/RigMapperDefinitionEditorToolkit.h
new file mode 100644
--- /dev/null
+++ b/Engine/Plugins/Experimental/Animation/RigMapper/Source/RigMapperEditor/Private/RigMapperDefinitionEditorToolkit.h
@@ -0,0 +1,68 @@
+// Copyright Epic Games, Inc. All Rights Reserved.
+
+#pragma once
+
+#include "IDetailCustomization.h"
+#include "RigMapperDefinition.h"
+
+#include "Toolkits/AssetEditorToolkit.h"
+
+enum class ERigMapperNodeType : uint8;
+class SRigMapperDefinitionStructureView;
+class SRigMapperDefinitionGraphEditor;
+
+/**
+ * The toolkit for the URigMapperDefinition asset editor
+ */
+class RIGMAPPEREDITOR_API FRigMapperDefinitionEditorToolkit : public FAssetEditorToolkit
+{
+public:
+	void Initialize(URigMapperDefinition* InDefinition, EToolkitMode::Type InMode, TSharedPtr<IToolkitHost> InToolkitHost);
+
+	//~ FAssetEditorToolkit Interface
+	virtual void RegisterTabSpawners(const TSharedRef<class FTabManager>& InTabManager) override;
+	virtual void UnregisterTabSpawners(const TSharedRef<class FTabManager>& InTabManager) override;
+	virtual FName GetToolkitFName() const override;
+	virtual FText GetBaseToolkitName() const override;
+	virtual FString GetWorldCentricTabPrefix() const override;
+	virtual FLinearColor GetWorldCentricTabColorScale() const override;
+	//~ IToolkit interface
+
+	class FDetailsViewCustomization : public IDetailCustomization
+	{
+	public:
+		/** Makes a new instance of this detail layout class for a specific detail view requesting it */
+		static TSharedRef<IDetailCustomization> MakeInstance();
+
+		FDetailsViewCustomization()
+		{}
+
+		// IDetailCustomization interface
+		virtual void CustomizeDetails(IDetailLayoutBuilder& DetailLayout) override;
+	};
+	
+private:
+	TSharedRef<SDockTab> SpawnGraphTab(const FSpawnTabArgs& SpawnTabArgs);
+	TSharedRef<SDockTab> SpawnDetailsTab(const FSpawnTabArgs& SpawnTabArgs);
+	TSharedRef<SDockTab> SpawnStructureTab(const FSpawnTabArgs& SpawnTabArgs);
+
+	bool HandleIsPropertyVisible(const FPropertyAndParent& PropertyAndParent);
+	void HandleFinishedChangingProperties(const FPropertyChangedEvent& PropertyChangedEvent);
+	void HandleRigMapperDefinitionLoaded();
+	void HandleGraphSelectionChanged(const TSet<UObject*>& Nodes);
+	
+	void HandleStructureSelectionChanged(ESelectInfo::Type SelectInfo, TArray<FString> SelectedInputs, TArray<FString> SelectedFeatures, TArray<FString> SelectedOutputs);
+
+private:
+	static const FName DefinitionEditorGraphTabId;
+	static const FName DefinitionEditorStructureTabId;
+	static const FName DefinitionEditorDetailsTabId;
+
+	static const TMap<FName, ERigMapperNodeType> PropertyNameToNodeTypeMapping;
+	
+	TObjectPtr<URigMapperDefinition> Definition;
+
+	TSharedPtr<SRigMapperDefinitionGraphEditor> GraphEditor;
+	TSharedPtr<IDetailsView> DetailsView;
+	TSharedPtr<SRigMapperDefinitionStructureView> StructureView;
+};
diff --git a/Engine/Plugins/Experimental/Animation/RigMapper/Source/RigMapperEditor/Private/RigMapperEditorModule.cpp b/Engine/Plugins/Experimental/Animation/RigMapper/Source/RigMapperEditor/Private/RigMapperEditorModule.cpp
new file mode 100644
--- /dev/null
+++ b/Engine/Plugins/Experimental/Animation/RigMapper/Source/RigMapperEditor/Private/RigMapperEditorModule.cpp
@@ -0,0 +1,430 @@
+// Copyright Epic Games, Inc. All Rights Reserved.
+
+#include "RigMapperEditorModule.h"
+
+#include "RigMapperDefinition.h"
+#include "RigMapperDefinitionAssetTypeActions.h"
+#include "RigMapperGraph/RigMapperDefinitionEditorGraphSchema.h"
+#include "RigMapperLinkedDefinitionsAssetTypeActions.h"
+
+#include "AssetToolsModule.h"
+#include "ContentBrowserMenuContexts.h"
+#include "DesktopPlatformModule.h"
+#include "EdGraphUtilities.h"
+#include "EditorDirectories.h"
+#include "IDesktopPlatform.h"
+#include "MessageLogModule.h"
+#include "Brushes/SlateImageBrush.h"
+#include "Framework/Application/SlateApplication.h"
+#include "Modules/ModuleManager.h"
+#include "Styling/SlateStyle.h"
+#include "Styling/SlateStyleRegistry.h"
+#include "Toolkits/AssetEditorToolkitMenuContext.h"
+#include "ToolMenus.h"
+#include "Logging/MessageLog.h"
+#include "RigMapperGraph/RigMapperDefinitionEditorGraphNode.h"
+
+const FName FRigMapperEditorModule::AppIdentifier("RigMapperEditorApp");
+
+const FName FRigMapperEditorModule::MessageLogIdentifier("RigMapperEditor");
+
+#define LOCTEXT_NAMESPACE "RigMapperEditorModule"
+
+void FRigMapperEditorModule::StartupModule()
+{
+	RegisterSlateStyle();
+	
+	if (FModuleManager::Get().IsModuleLoaded("AssetTools"))
+	{
+		// Register asset, editor & asset actions
+		RigMapperDefinitionAssetTypeActions = MakeShared<FRigMapperDefinitionAssetTypeActions>();
+		RigMapperLinkedDefinitionsAssetTypeActions = MakeShared<FRigMapperLinkedDefinitionsAssetTypeActions>();
+		
+		FAssetToolsModule::GetModule().Get().RegisterAssetTypeActions(RigMapperDefinitionAssetTypeActions.ToSharedRef());
+		FAssetToolsModule::GetModule().Get().RegisterAssetTypeActions(RigMapperLinkedDefinitionsAssetTypeActions.ToSharedRef());
+	}
+	
+	FToolMenuOwnerScoped OwnerScoped(this);
+
+	RegisterRigMapperDefinitionContextMenuEntries();
+	RegisterRigMapperDefinitionToolbarEntries();
+	RegisterRigMapperLinkedDefinitionsContextMenuEntries();
+	RegisterRigMapperLinkedDefinitionsToolbarEntries();
+
+	FMessageLogModule& MessageLogModule = FModuleManager::LoadModuleChecked<FMessageLogModule>("MessageLog");
+	MessageLogModule.RegisterLogListing(MessageLogIdentifier, LOCTEXT("RigMapperEditorMessageLog", "Rig Mapper"));
+
+	RigMapperDefinitionGraphEditorNodeFactory = MakeShareable(new URigMapperDefinitionEditorGraphNode::NodeFactory());
+	FEdGraphUtilities::RegisterVisualNodeFactory(RigMapperDefinitionGraphEditorNodeFactory);
+}
+
+void FRigMapperEditorModule::ShutdownModule()
+{
+	FSlateStyleRegistry::UnRegisterSlateStyle(*Style.Get());
+	Style.Reset();
+	
+	if (FModuleManager::Get().IsModuleLoaded("AssetTools"))
+	{
+		FAssetToolsModule::GetModule().Get().UnregisterAssetTypeActions(RigMapperDefinitionAssetTypeActions.ToSharedRef());
+	}
+
+	FMessageLogModule& MessageLogModule = FModuleManager::LoadModuleChecked<FMessageLogModule>("MessageLog");
+	MessageLogModule.UnregisterLogListing(MessageLogIdentifier);
+
+	if (RigMapperDefinitionGraphEditorNodeFactory.IsValid())
+	{
+		FEdGraphUtilities::UnregisterVisualNodeFactory(RigMapperDefinitionGraphEditorNodeFactory);
+	}
+}
+
+void FRigMapperEditorModule::RegisterSlateStyle()
+{
+	Style = MakeShareable(new FSlateStyleSet("RigMapperEditorStyle"));
+	Style->SetContentRoot(FPaths::EngineContentDir() / TEXT("Editor/Slate"));
+	Style->SetCoreContentRoot(FPaths::EngineContentDir() / TEXT("Slate"));
+	
+	Style->Set("ClassIcon.RigMapperDefinition", new FSlateVectorImageBrush(
+		Style->RootToContentDir("Starship/Common/Blueprint", TEXT(".svg")),
+		FVector2D(16.0f, 16.0f)));
+	Style->Set("ClassThumbnail.RigMapperDefinition", new FSlateVectorImageBrush(
+		Style->RootToContentDir("Starship/Common/Blueprint", TEXT(".svg"))
+		, FVector2D(64.0f, 64.0f)));
+
+	Style->Set("ClassIcon.RigMapperLinkedDefinitions", new FSlateVectorImageBrush(
+		Style->RootToContentDir("Starship/Common/Struct", TEXT(".svg")),
+		FVector2D(16.0f, 16.0f)));
+	Style->Set("ClassThumbnail.RigMapperLinkedDefinitions", new FSlateVectorImageBrush(
+		Style->RootToContentDir("Starship/Common/Struct", TEXT(".svg"))
+		, FVector2D(64.0f, 64.0f)));
+	
+	FSlateStyleRegistry::RegisterSlateStyle(*Style.Get());
+}
+
+void FRigMapperEditorModule::RegisterRigMapperDefinitionContextMenuEntries()
+{
+	UToolMenu* Menu = UToolMenus::Get()->ExtendMenu("ContentBrowser.AssetContextMenu.RigMapperDefinition");
+
+	const FToolMenuInsert SectionInsertLocation("GetAssetActions", EToolMenuInsertType::After);
+	FToolMenuSection& Section = Menu->AddSection("RigMapperActions", INVTEXT("Rig Mapper"), SectionInsertLocation);
+
+	Section.AddMenuEntry(
+		"RigMapperDefinition_LoadFromJson",
+		LOCTEXT("RigMapperDefinitionContextMenu_LoadFromJson_Label", "Load from Json"),
+		LOCTEXT("RigMapperDefinitionContextMenu_LoadFromJson_Tooltip", "Reload this definition from the given json file"),
+		FSlateIcon(FAppStyle::GetAppStyleSetName(), "Icons.Toolbar.Import"),
+		FToolMenuExecuteAction::CreateStatic(&FRigMapperEditorModule::LoadFromJson, false)
+	);
+	Section.AddMenuEntry(
+		"RigMapperDefinition_ExportToJson",
+		LOCTEXT("RigMapperDefinitionContextMenu_ExportToJson_Label", "Export to Json"),
+		LOCTEXT("RigMapperDefinitionContextMenu_ExportToJson_Tooltip", "Export this definition as a json file"),
+		FSlateIcon(FAppStyle::GetAppStyleSetName(), "Icons.Toolbar.Export"),
+		FToolMenuExecuteAction::CreateStatic(&FRigMapperEditorModule::ExportToJson, false)
+	);
+	Section.AddMenuEntry(
+		"RigMapperDefinition_Validate",
+		LOCTEXT("RigMapperDefinitionContextMenu_Validate_Label", "Validate Definition"),
+		LOCTEXT("RigMapperDefinitionContextMenu_Validate_Tooltip", "Check if the definition has any noticable issue"),
+		FSlateIcon(FAppStyle::GetAppStyleSetName(), "Icons.Validate"),
+		FToolMenuExecuteAction::CreateStatic(&FRigMapperEditorModule::ValidateDefinition, false)
+	);
+}
+
+void FRigMapperEditorModule::RegisterRigMapperDefinitionToolbarEntries()
+{
+	UToolMenu* Menu = UToolMenus::Get()->ExtendMenu("AssetEditor.RigMapperDefinitionEditor.ToolBar");
+	FToolMenuSection& Section = Menu->FindOrAddSection("AssetEditorActions");
+
+	Section.AddMenuEntry(
+		"RigMapperDefinition_LoadFromJson",
+		LOCTEXT("RigMapperDefinitionToolbar_LoadFromJson_Label", "Load"),
+		LOCTEXT("RigMapperDefinitionToolbar_LoadFromJson_Tooltip", "Reload this definition from the given json file"),
+		FSlateIcon(FAppStyle::GetAppStyleSetName(), "Icons.Toolbar.Import"),
+		FToolMenuExecuteAction::CreateStatic(&FRigMapperEditorModule::LoadFromJson, true)
+	);
+	Section.AddMenuEntry(
+		"RigMapperDefinition_ExportToJson",
+		LOCTEXT("RigMapperDefinitionToolbar_ExportToJson_Label", "Export"),
+		LOCTEXT("RigMapperDefinitionToolbar_ExportToJson_Tooltip", "Export this definition as a json file"),
+		FSlateIcon(FAppStyle::GetAppStyleSetName(), "Icons.Toolbar.Export"),
+		FToolMenuExecuteAction::CreateStatic(&FRigMapperEditorModule::ExportToJson, true)
+	);
+	Section.AddMenuEntry(
+		"RigMapperDefinition_Validate",
+		LOCTEXT("RigMapperDefinitionToolbar_Validate_Label", "Validate"),
+		LOCTEXT("RigMapperDefinitionToolbar_Validate_Tooltip", "Check if the definition has any noticable issue"),
+		FSlateIcon(FAppStyle::GetAppStyleSetName(), "Icons.Validate"),
+		FToolMenuExecuteAction::CreateStatic(&FRigMapperEditorModule::ValidateDefinition, true)
+	);
+}
+
+
+void FRigMapperEditorModule::RegisterRigMapperLinkedDefinitionsContextMenuEntries()
+{
+	UToolMenu* Menu = UToolMenus::Get()->ExtendMenu("ContentBrowser.AssetContextMenu.RigMapperLinkedDefinitions");
+
+	const FToolMenuInsert SectionInsertLocation("GetAssetActions", EToolMenuInsertType::After);
+	FToolMenuSection& Section = Menu->AddSection("RigMapperActions", INVTEXT("Rig Mapper"), SectionInsertLocation);
+	
+	Section.AddMenuEntry(
+		"RigMapperLinkedDefinitions_Bake",
+		LOCTEXT("RigMapperLinkedDefinitionsContextMenu_Bake_Label", "Bake Definitions"),
+		LOCTEXT("RigMapperLinkedDefinitionsContextMenu_Bake_Tooltip", "The linked definitions will be baked to the output definition"),
+		FSlateIcon(FAppStyle::GetAppStyleSetName(), "AssetEditor.Apply"),
+		FToolMenuExecuteAction::CreateStatic(&FRigMapperEditorModule::BakeDefinitions, false)
+	);
+	Section.AddMenuEntry(
+		"RigMapperLinkedDefinitions_Validate",
+		LOCTEXT("RigMapperLinkedDefinitionsContextMenu_Validate_Label", "Validate Linked Definitions"),
+		LOCTEXT("RigMapperLinkedDefinitionsContextMenu_Validate_Tooltip", "Check if the definitions are valid and link together correctly"),
+		FSlateIcon(FAppStyle::GetAppStyleSetName(), "Icons.Validate"),
+		FToolMenuExecuteAction::CreateStatic(&FRigMapperEditorModule::ValidateLinkedDefinitions, false)
+	);
+}
+
+void FRigMapperEditorModule::RegisterRigMapperLinkedDefinitionsToolbarEntries()
+{
+	UToolMenu* Menu = UToolMenus::Get()->ExtendMenu("AssetEditor.RigMapperLinkedDefinitionsEditor.ToolBar");
+	FToolMenuSection& Section = Menu->FindOrAddSection("AssetEditorActions");
+
+	// todo: bake to new
+	
+	Section.AddEntry(FToolMenuEntry::InitToolBarButton(
+		"RigMapperLinkedDefinitions_Bake",
+		FToolMenuExecuteAction::CreateStatic(&FRigMapperEditorModule::BakeDefinitions, true),
+		LOCTEXT("RigMapperLinkedDefinitionsToolbar_Bake_Label", "Bake"),
+		LOCTEXT("RigMapperLinkedDefinitionsToolbar_Bake_Tooltip", "The linked definitions will be baked to the output definition"),
+		FSlateIcon(FAppStyle::GetAppStyleSetName(), "AssetEditor.Apply")
+	));
+
+	Section.AddEntry(FToolMenuEntry::InitToolBarButton(
+		"RigMapperLinkedDefinitions_Validate",
+		FToolMenuExecuteAction::CreateStatic(&FRigMapperEditorModule::ValidateLinkedDefinitions, true),
+		LOCTEXT("RigMapperLinkedDefinitionsToolbar_Validate_Label", "Validate"),
+		LOCTEXT("RigMapperLinkedDefinitionsToolbar_Validate_Tooltip", "Check if the definitions are valid and link together correctly"),
+		FSlateIcon(FAppStyle::GetAppStyleSetName(), "Icons.Validate")
+	));
+}
+
+bool FRigMapperEditorModule::GetUserPickedPath(FString& OutPath, bool bImport)
+{
+	TArray<FString> Filenames;
+	IDesktopPlatform* DesktopPlatform = FDesktopPlatformModule::Get();
+	
+	bool bOpen;
+	
+	if (bImport)
+	{
+		bOpen = DesktopPlatform->OpenFileDialog(
+			FSlateApplication::Get().FindBestParentWindowHandleForDialogs(nullptr),
+			NSLOCTEXT("RigMapperDefinitionIO", "ImportFromJson", "Import Definition from Json...").ToString(), 
+			FEditorDirectories::Get().GetLastDirectory(ELastDirectory::GENERIC_IMPORT),
+			TEXT(""), 
+			TEXT("json (*.json)|*.json|"),
+			EFileDialogFlags::None,
+			Filenames
+		);
+	}
+	else
+	{
+		constexpr bool bDirectory = true;
+		
+		if (bDirectory)
+		{
+			return DesktopPlatform->OpenDirectoryDialog(
+				FSlateApplication::Get().FindBestParentWindowHandleForDialogs(nullptr),
+				NSLOCTEXT("RigMapperDefinitionIO", "ExportToJson", "Export Definition to Json...").ToString(), 
+				FEditorDirectories::Get().GetLastDirectory(ELastDirectory::GENERIC_EXPORT),
+				OutPath
+			);
+		}
+		bOpen = DesktopPlatform->SaveFileDialog(
+			FSlateApplication::Get().FindBestParentWindowHandleForDialogs(nullptr),
+			NSLOCTEXT("RigMapperDefinitionIO", "ExportToJson", "Export Definition to Json...").ToString(), 
+			FEditorDirectories::Get().GetLastDirectory(ELastDirectory::GENERIC_EXPORT),
+			TEXT(""), 
+			TEXT("json (*.json)|*.json|"),
+			EFileDialogFlags::None,
+			Filenames
+		);
+	}
+	
+	if (bOpen && !Filenames.IsEmpty())
+	{
+		OutPath = Filenames[0];
+	}
+
+	return bOpen;
+}
+
+void FRigMapperEditorModule::LoadFromJson(const FToolMenuContext& ToolMenuContext, bool bFromAssetEditor)
+{
+	const TArray<URigMapperDefinition*> Definitions = GetDefinitionsFromContext(ToolMenuContext, bFromAssetEditor);
+	
+	FFilePath JsonFile;
+
+	if (!Definitions.IsEmpty() && GetUserPickedPath(JsonFile.FilePath))
+	{
+		for (URigMapperDefinition* Definition : Definitions)
+		{
+			if (!Definition->LoadFromJsonFile(JsonFile))
+			{
+				FMessageLog MessageLog(MessageLogIdentifier);
+
+				FFormatNamedArguments Arguments;
+				Arguments.Add(TEXT("ObjectName"), FText::FromString(Definition->GetName()));
+	
+				MessageLog.Error(FText::Format(LOCTEXT("RigMapperEditorActions_ImportFailed", "Failed to load definition \"{ObjectName}\" from Json. See output log for more details"), Arguments));
+				MessageLog.Open(EMessageSeverity::Error);
+			}
+		}
+	}
+}
+
+void FRigMapperEditorModule::ExportToJson(const FToolMenuContext& ToolMenuContext, bool bFromAssetEditor)
+{
+	const TArray<URigMapperDefinition*> Definitions = GetDefinitionsFromContext(ToolMenuContext, bFromAssetEditor);
+	
+	FDirectoryPath Path;
+	
+	if (!Definitions.IsEmpty() && GetUserPickedPath(Path.Path, false))
+	{
+		for (URigMapperDefinition* Definition : Definitions)
+		{
+			FFilePath ActualFilePath;
+			ActualFilePath.FilePath = FPaths::Combine(Path.Path, Definition->GetName() + ".json");
+
+			if (!Definition->ExportAsJsonFile(ActualFilePath))
+			{
+				FMessageLog MessageLog(MessageLogIdentifier);
+
+				FFormatNamedArguments Arguments;
+				Arguments.Add(TEXT("ObjectName"), FText::FromString(Definition->GetName()));
+	
+				MessageLog.Error(FText::Format(LOCTEXT("RigMapperEditorActions_ExportFailed", "Failed to export definition \"{ObjectName}\". See output log for more details"), Arguments));
+				MessageLog.Open(EMessageSeverity::Error);
+			}
+		}
+	}
+}
+
+void FRigMapperEditorModule::ValidateDefinition(const FToolMenuContext& ToolMenuContext, bool bFromAssetEditor)
+{
+	const TArray<URigMapperDefinition*> Definitions = GetDefinitionsFromContext(ToolMenuContext, bFromAssetEditor);
+
+	for (URigMapperDefinition* Definition : Definitions)
+	{
+		if (!Definition->Validate())
+		{
+			FMessageLog MessageLog(MessageLogIdentifier);
+
+			FFormatNamedArguments Arguments;
+			Arguments.Add(TEXT("ObjectName"), FText::FromString(Definition->GetName()));
+	
+			MessageLog.Error(FText::Format(LOCTEXT("RigMapperEditorActions_ValidateFailed", "Failed to validate definition \"{ObjectName}\". See output log for more details"), Arguments));
+			MessageLog.Open(EMessageSeverity::Error);
+		}
+	}
+}
+
+void FRigMapperEditorModule::BakeDefinitions(const FToolMenuContext& ToolMenuContext, bool bFromAssetEditor)
+{
+	TArray<URigMapperLinkedDefinitions*> Definitions = GetLinkedDefinitionsFromContext(ToolMenuContext, bFromAssetEditor);
+	
+	for (URigMapperLinkedDefinitions* LinkedDefinition : Definitions)
+	{
+		if (!LinkedDefinition->BakeDefinitions())
+		{
+			FMessageLog MessageLog(MessageLogIdentifier);
+
+			FFormatNamedArguments Arguments;
+			Arguments.Add(TEXT("ObjectName"), FText::FromString(LinkedDefinition->GetName()));
+	
+			MessageLog.Error(FText::Format(LOCTEXT("RigMapperEditorActions_BakeFailed", "Failed to bake definitions from Linked Defintion \"{ObjectName}\". See output log for more details"), Arguments));
+			MessageLog.Open(EMessageSeverity::Error);
+		}
+	}
+}
+
+void FRigMapperEditorModule::ValidateLinkedDefinitions(const FToolMenuContext& ToolMenuContext, bool bFromAssetEditor)
+{
+	TArray<URigMapperLinkedDefinitions*> Definitions = GetLinkedDefinitionsFromContext(ToolMenuContext, bFromAssetEditor);
+
+	for (URigMapperLinkedDefinitions* LinkedDefinition : Definitions)
+	{
+		if (!LinkedDefinition->Validate())
+		{
+			FMessageLog MessageLog(MessageLogIdentifier);
+
+			FFormatNamedArguments Arguments;
+			Arguments.Add(TEXT("ObjectName"), FText::FromString(LinkedDefinition->GetName()));
+	
+			MessageLog.Error(FText::Format(LOCTEXT("RigMapperEditorActions_ValidateLinkFailed", "Failed to validate definitions from Linked Defintion \"{ObjectName}\". See output log for more details"), Arguments));
+			MessageLog.Open(EMessageSeverity::Error);
+		}
+	}
+}
+
+TArray<URigMapperLinkedDefinitions*> FRigMapperEditorModule::GetLinkedDefinitionsFromContext(const FToolMenuContext& ToolMenuContext, bool bFromAssetEditor)
+{
+	TArray<URigMapperLinkedDefinitions*> Definitions;
+	
+	if (bFromAssetEditor)
+	{
+		const UAssetEditorToolkitMenuContext* ToolkitMenuContext = ToolMenuContext.FindContext<UAssetEditorToolkitMenuContext>();
+	
+		if (ToolkitMenuContext && ToolkitMenuContext->Toolkit.IsValid())
+		{
+			for (UObject* Object : ToolkitMenuContext->GetEditingObjects())
+			{
+				if (URigMapperLinkedDefinitions* Definition = Cast<URigMapperLinkedDefinitions>(Object))
+				{
+					Definitions.Add(Definition);
+				}
+			}
+		}
+		
+	}
+	else
+	{
+		if (const UContentBrowserAssetContextMenuContext* Context = UContentBrowserAssetContextMenuContext::FindContextWithAssets(ToolMenuContext))
+		{
+			Definitions = Context->LoadSelectedObjects<URigMapperLinkedDefinitions>();
+		}
+	}
+	return Definitions;
+}
+
+TArray<URigMapperDefinition*> FRigMapperEditorModule::GetDefinitionsFromContext(const FToolMenuContext& ToolMenuContext, bool bFromAssetEditor)
+{
+	TArray<URigMapperDefinition*> Definitions;
+	
+	if (bFromAssetEditor)
+	{
+		const UAssetEditorToolkitMenuContext* ToolkitMenuContext = ToolMenuContext.FindContext<UAssetEditorToolkitMenuContext>();
+	
+		if (ToolkitMenuContext && ToolkitMenuContext->Toolkit.IsValid())
+		{
+			for (UObject* Object : ToolkitMenuContext->GetEditingObjects())
+			{
+				if (URigMapperDefinition* Definition = Cast<URigMapperDefinition>(Object))
+				{
+					Definitions.Add(Definition);
+				}
+			}
+		}
+		
+	}
+	else
+	{
+		if (const UContentBrowserAssetContextMenuContext* Context = UContentBrowserAssetContextMenuContext::FindContextWithAssets(ToolMenuContext))
+		{
+			Definitions = Context->LoadSelectedObjects<URigMapperDefinition>();
+		}
+	}
+	return Definitions;
+}
+
+#undef LOCTEXT_NAMESPACE
+
+IMPLEMENT_MODULE(FRigMapperEditorModule, RigMapperEditor)
diff --git a/Engine/Plugins/Experimental/Animation/RigMapper/Source/RigMapperEditor/Private/RigMapperEditorModule.h b/Engine/Plugins/Experimental/Animation/RigMapper/Source/RigMapperEditor/Private/RigMapperEditorModule.h
new file mode 100644
--- /dev/null
+++ b/Engine/Plugins/Experimental/Animation/RigMapper/Source/RigMapperEditor/Private/RigMapperEditorModule.h
@@ -0,0 +1,56 @@
+// Copyright Epic Games, Inc. All Rights Reserved.
+
+#pragma once
+
+#include "CoreMinimal.h"
+#include "ToolMenuContext.h"
+#include "Modules/ModuleInterface.h"
+
+struct FGraphPanelNodeFactory;
+class URigMapperDefinition;
+class FRigMapperLinkedDefinitionsAssetTypeActions;
+class URigMapperLinkedDefinitions;
+class FSpawnTabArgs;
+class SDockTab;
+class FRigMapperDefinitionAssetTypeActions;
+class FSlateStyleSet;
+
+/**
+ * FRigMapperEditorModule
+ * The Editor Module for the Rig Mapper plugin
+ */
+class FRigMapperEditorModule : public IModuleInterface
+{
+public:
+	//~ Begin IModuleInterface interface
+	virtual void StartupModule() override;
+	virtual void ShutdownModule() override;
+	//~ End IModuleInterface interface
+
+	static void RegisterRigMapperDefinitionContextMenuEntries();
+	static void RegisterRigMapperDefinitionToolbarEntries();
+	static void RegisterRigMapperLinkedDefinitionsContextMenuEntries();
+	static void RegisterRigMapperLinkedDefinitionsToolbarEntries();
+	
+private:
+	void RegisterSlateStyle();
+	static bool GetUserPickedPath(FString& OutPath, bool bImport = true);
+	static void LoadFromJson(const FToolMenuContext& ToolMenuContext, bool bFromAssetEditor);
+	static void ExportToJson(const FToolMenuContext& ToolMenuContext, bool bFromAssetEditor);
+	static void ValidateDefinition(const FToolMenuContext& ToolMenuContext, bool bFromAssetEditor);
+	static void BakeDefinitions(const FToolMenuContext& ToolMenuContext, bool bFromAssetEditor);
+	static void ValidateLinkedDefinitions(const FToolMenuContext& ToolMenuContext, bool bFromAssetEditor);
+	static TArray<URigMapperDefinition*> GetDefinitionsFromContext(const FToolMenuContext& ToolMenuContext, bool bFromAssetEditor);
+	static TArray<URigMapperLinkedDefinitions*> GetLinkedDefinitionsFromContext(const FToolMenuContext& ToolMenuContext, bool bFromAssetEditor);
+	
+public:
+	static const FName AppIdentifier;
+	static const FName MessageLogIdentifier;
+	
+private:
+	TSharedPtr<FSlateStyleSet> Style;
+	
+	TSharedPtr<FRigMapperDefinitionAssetTypeActions> RigMapperDefinitionAssetTypeActions;
+	TSharedPtr<FRigMapperLinkedDefinitionsAssetTypeActions> RigMapperLinkedDefinitionsAssetTypeActions;
+	TSharedPtr<FGraphPanelNodeFactory> RigMapperDefinitionGraphEditorNodeFactory;
+};
diff --git a/Engine/Plugins/Experimental/Animation/RigMapper/Source/RigMapperEditor/Private/RigMapperEditorSubsystem.cpp b/Engine/Plugins/Experimental/Animation/RigMapper/Source/RigMapperEditor/Private/RigMapperEditorSubsystem.cpp
new file mode 100644
--- /dev/null
+++ b/Engine/Plugins/Experimental/Animation/RigMapper/Source/RigMapperEditor/Private/RigMapperEditorSubsystem.cpp
@@ -0,0 +1,1004 @@
+// Copyright Epic Games, Inc. All Rights Reserved.
+
+#include "RigMapperEditorSubsystem.h"
+
+#include "Engine/SkeletalMesh.h"
+#include "Animation/AnimSequence.h"
+#include "Misc/FrameNumber.h"
+#include "AssetToolsModule.h"
+#include "LevelSequence.h"
+#include "SequencerUtilities.h"
+#include "AssetRegistry/AssetRegistryModule.h"
+#include "Misc/FileHelper.h"
+#include "Sequencer/MovieSceneControlRigParameterSection.h"
+#include "Sequencer/MovieSceneControlRigParameterTrack.h"
+#include "Tracks/MovieSceneSpawnTrack.h"
+#include "Units/Execution/RigUnit_PrepareForExecution.h"
+#include "Interfaces//ITargetPlatformManagerModule.h"
+
+UAnimSequence* CreateNewAnimSequence(const FString& BaseName, const FString& Path, USkeletalMesh* TargetMesh)
+{
+	FString UniqueAssetName = BaseName + "_" + TargetMesh->GetName();
+	
+	FString UniquePackageName;
+	const FString BasePackageName = FString(Path / UniqueAssetName);
+
+	const FAssetToolsModule& AssetToolsModule = FModuleManager::Get().LoadModuleChecked<FAssetToolsModule>("AssetTools");
+	AssetToolsModule.Get().CreateUniqueAssetName(BasePackageName, "", /*out*/ UniquePackageName, /*out*/ UniqueAssetName);
+
+	const bool bCreatePackage = !UniquePackageName.IsEmpty() && !UniqueAssetName.IsEmpty();
+	UPackage* Package = bCreatePackage ? CreatePackage(*UniquePackageName) : GetTransientPackage();
+	
+	UAnimSequence* AnimSequence = FindObject<UAnimSequence>(Package, *UniqueAssetName);
+	if (!AnimSequence)
+	{
+		AnimSequence = NewObject<UAnimSequence>(Package, *UniqueAssetName, bCreatePackage ? RF_Public | RF_Standalone : RF_NoFlags);
+	}
+
+	AnimSequence->SetSkeleton(TargetMesh->GetSkeleton());
+	AnimSequence->SetPreviewMesh(TargetMesh);
+
+	FAssetRegistryModule::AssetCreated(AnimSequence);
+	
+	return AnimSequence;
+}
+
+FGuid CreateSequencerSpawnable(UMovieScene* MovieScene, UObject* ObjectToSpawn)
+{
+	if (!MovieScene || !ObjectToSpawn)
+	{
+		return FGuid();
+	}
+
+	UE::Sequencer::FCreateBindingParams CreateBindingParams;
+	CreateBindingParams.bAllowCustomBinding = true;
+	CreateBindingParams.bSpawnable = true;
+	return FSequencerUtilities::CreateOrReplaceBinding(nullptr, MovieScene->GetTypedOuter<UMovieSceneSequence>(), ObjectToSpawn, CreateBindingParams);
+}
+
+UMovieSceneControlRigParameterSection* CreateControlRigSequence(const FString& PackageName, const FString& AssetName, const FFrameRate& FrameRate, const TSubclassOf<UControlRig>& ControlRigClass, USkeletalMesh* SkeletalMesh)
+{
+	if (!SkeletalMesh || !ControlRigClass)
+	{
+		UE_LOG(LogTemp, Error, TEXT("Could not create level sequence and control rig track to record, invalid skeletal mesh or control rig class provided"))
+		return nullptr;
+	}
+
+	FString UniquePackageName;
+	FString UniqueAssetName = AssetName + "_" + SkeletalMesh->GetName();;
+	const FString BasePackageName = FString(PackageName / UniqueAssetName);
+
+	const FAssetToolsModule& AssetToolsModule = FModuleManager::Get().LoadModuleChecked<FAssetToolsModule>("AssetTools");
+	AssetToolsModule.Get().CreateUniqueAssetName(BasePackageName, "", /*out*/ UniquePackageName, /*out*/ UniqueAssetName);
+	
+	const bool bCreatePackage = !UniquePackageName.IsEmpty() && !UniqueAssetName.IsEmpty();
+	UPackage* Package = bCreatePackage ? CreatePackage(*UniquePackageName) : GetTransientPackage();
+
+	ULevelSequence* Sequence = FindObject<ULevelSequence>(Package, *UniqueAssetName);
+	if (!Sequence)
+	{
+		Sequence = NewObject<ULevelSequence>(Package, *UniqueAssetName, bCreatePackage ? RF_Public | RF_Standalone : RF_NoFlags);
+		Sequence->Initialize();
+	}
+
+	UMovieScene* MovieScene = Sequence->GetMovieScene();
+	check(MovieScene);
+
+	MovieScene->Modify();
+
+	MovieScene->SetDisplayRate(FrameRate);
+	
+	const FGuid NewGuid = Sequence->AllowsSpawnableObjects() ? CreateSequencerSpawnable(MovieScene, SkeletalMesh) : FGuid();
+	if (!NewGuid.IsValid())
+	{
+		UE_LOG(LogTemp, Error, TEXT("Failed to create sequencer spawnable from the given skeletal mesh asset"))
+		return nullptr;
+	}
+	
+	UMovieSceneControlRigParameterTrack* Track = Cast<UMovieSceneControlRigParameterTrack>(
+		MovieScene->AddTrack(UMovieSceneControlRigParameterTrack::StaticClass(), NewGuid));
+	check(Track);
+
+	FString TrackName = ControlRigClass->GetName();
+	TrackName.RemoveFromEnd(TEXT("_C"));
+
+	Track->SetTrackName(*TrackName);
+	Track->SetDisplayName(FText::FromString(TrackName));
+
+	UControlRig* ControlRig = NewObject<UControlRig>(Track, ControlRigClass, *TrackName, RF_Transactional);
+
+	ControlRig->Modify();
+	ControlRig->Initialize();
+	ControlRig->Evaluate_AnyThread();
+
+	Track->Modify();
+	
+	UMovieSceneSection* NewSection = Track->CreateControlRigSection(0, ControlRig, true);
+	check(NewSection);
+	
+	NewSection->Modify();
+	NewSection->SetRange(TRange<FFrameNumber>());
+
+	FAssetRegistryModule::AssetCreated(Sequence);
+	Sequence->MarkPackageDirty();
+	
+	return Cast<UMovieSceneControlRigParameterSection>(NewSection);
+}
+
+ULevelSequence* URigMapperEditorSubsystem::GetSequenceFromSection(const UMovieSceneControlRigParameterSection* Section)
+{
+	if (Section)
+	{
+		return Cast<ULevelSequence>(Section->GetTypedOuter(ULevelSequence::StaticClass()));
+	}
+	return nullptr;
+}
+
+TArray<UMovieSceneControlRigParameterSection*> URigMapperEditorSubsystem::GetSectionsFromSequence(ULevelSequence* Sequence)
+{
+	TArray<UMovieSceneControlRigParameterSection*> Sections;
+	
+	if (Sequence && Sequence->GetMovieScene())
+	{
+		for (const FMovieSceneBinding& Binding : Sequence->GetMovieScene()->GetBindings())
+		{
+			for (UMovieSceneTrack* Track : Binding.GetTracks())
+			{
+				if (Track->IsA(UMovieSceneControlRigParameterTrack::StaticClass()))
+				{
+					for (UMovieSceneSection* Section : Track->GetAllSections())
+					{
+						Sections.Add(Cast<UMovieSceneControlRigParameterSection>(Section));
+					}
+				}
+			}
+		}
+	}
+	return Sections;
+}
+
+FFrameRate URigMapperEditorSubsystem::GetAnimSequenceRate(const UAnimSequence* AnimSequence)
+{
+	return AnimSequence ? AnimSequence->GetTargetSamplingFrameRate(GetTargetPlatformManager()->GetRunningTargetPlatform()) : FFrameRate();
+}
+
+void URigMapperEditorSubsystem::SetAnimSequenceRate(UAnimSequence* AnimSequence, FFrameRate FrameRate, bool bSetImportProperties)
+{
+	if (AnimSequence && FrameRate.IsValid())
+	{
+		if (bSetImportProperties)
+		{
+			AnimSequence->ImportFileFramerate = FrameRate.AsDecimal();
+			AnimSequence->ImportResampleFramerate = FMath::RoundToInt32(FrameRate.AsDecimal());
+		}
+		
+		IAnimationDataController& Controller = AnimSequence->GetController();
+		Controller.OpenBracket(INVTEXT("Set Sequence FrameRate"));
+		Controller.InitializeModel();
+		Controller.SetFrameRate(FrameRate, false);
+
+		Controller.CloseBracket();	
+	
+		AnimSequence->MarkPackageDirty();
+	}
+}
+
+bool URigMapperEditorSubsystem::ConvertCurveValuesToCsv(TArray<FString>& CurveNames, const TArray<FFrameTime>& FrameTimes, TArray<FPoseValues>& CurveValuesPerFrame, const FFilePath& OutputFile, const TArray<URigMapperDefinition*>& Definitions, bool bOutputIntermediateCsvFiles)
+{
+	TArray<TSharedPtr<FacialRigMapping::FRigMapper>> RigMappers;
+	if (!InitializeRigMapping(Definitions, RigMappers))
+	{
+		// todo: warn
+		return false;
+	}
+
+	const FString BaseFilePath = FPaths::GetBaseFilename(OutputFile.FilePath, false);;
+	
+	for (int32 RigIndex = 0; RigIndex < RigMappers.Num(); RigIndex++)
+	{
+		if (bOutputIntermediateCsvFiles)
+		{
+			FFilePath CurrentOutputFile;
+			CurrentOutputFile.FilePath = FString::Printf(TEXT("%s_%d.%s"), *BaseFilePath, RigIndex, *FPaths::GetExtension(OutputFile.FilePath));
+			
+			WriteCurveValuesToCsv(CurrentOutputFile, CurveNames, FrameTimes, CurveValuesPerFrame);
+		}
+		
+		EvaluateRigMapping(CurveNames, FrameTimes, CurveValuesPerFrame, RigMappers[RigIndex]);
+	}
+	WriteCurveValuesToCsv(OutputFile, CurveNames, FrameTimes, CurveValuesPerFrame);
+	
+	return true;
+}
+
+bool URigMapperEditorSubsystem::ConvertCsv(const FFilePath& InputFile, const FFilePath& OutputFile, const TArray<URigMapperDefinition*>& Definitions, bool bOutputIntermediateCsvFiles)
+{
+	TArray<FString> CurveNames;
+	TArray<FFrameTime> FrameTimes;
+		TArray<FPoseValues> CurveValuesPerFrame;
+	
+	if (!LoadCurveValuesFromCsv(InputFile, CurveNames, FrameTimes, CurveValuesPerFrame))
+	{
+		return false;
+	}
+
+	return ConvertCurveValuesToCsv(CurveNames, FrameTimes, CurveValuesPerFrame, OutputFile, Definitions, bOutputIntermediateCsvFiles);
+}
+
+bool URigMapperEditorSubsystem::ConvertCurveValuesToAnimSequence(TArray<FString>& CurveNames, const TArray<FFrameTime>& FrameTimes, TArray<FPoseValues>& CurveValuesPerFrame, UAnimSequence* Target, const TArray<URigMapperDefinition*>& Definitions)
+{
+	TArray<TSharedPtr<FacialRigMapping::FRigMapper>> RigMappers;
+	if (!InitializeRigMapping(Definitions, RigMappers))
+	{
+		// todo: warn
+		return false;
+	}
+	
+	for (int32 RigIndex = 0; RigIndex < RigMappers.Num(); RigIndex++)
+	{
+		EvaluateRigMapping(CurveNames, FrameTimes, CurveValuesPerFrame, RigMappers[RigIndex]);
+	}
+	
+	return AddCurveValuesToAnimSequence(Target, CurveNames, FrameTimes, CurveValuesPerFrame);
+}
+
+bool URigMapperEditorSubsystem::AddCurveValuesToAnimSequence(UAnimSequence* Target, const TArray<FString>& CurveNames, const TArray<FFrameTime>& FrameTimes, const TArray<FPoseValues>& CurveValuesPerFrame)
+{	
+	if (CurveNames.IsEmpty() || FrameTimes.IsEmpty() || CurveValuesPerFrame.Num() != FrameTimes.Num())
+	{
+		return false;
+	}
+
+	const FFrameRate FrameRate = GetAnimSequenceRate(Target);
+	if (!FrameRate.IsValid())
+	{
+		return false;
+	}
+
+	Target->ResetAnimation();
+	
+	IAnimationDataController& Controller = Target->GetController();
+	Controller.OpenBracket(INVTEXT("Populate New Anim"));
+	Controller.InitializeModel();
+	
+	Controller.SetNumberOfFrames(FrameTimes.Last().CeilToFrame().Value - FrameTimes[0].FloorToFrame().Value + 1, false);
+
+	// AnimSequence should start on frame 0
+	const float TimeOffset = FrameRate.AsSeconds(FrameTimes[0]);
+	
+	for (int32 CurveIndex = 0; CurveIndex < CurveNames.Num(); CurveIndex++)
+	{
+		FAnimationCurveIdentifier CurveId = UAnimationCurveIdentifierExtensions::GetCurveIdentifier(Target->GetSkeleton(), *CurveNames[CurveIndex], ERawCurveTrackTypes::RCT_Float);
+		// todo: remap curve name from outputs
+
+		if (CurveId.IsValid())
+		{
+			Controller.AddCurve(CurveId, 4, false);
+
+			for (int32 FrameIndex = 0; FrameIndex < FrameTimes.Num(); FrameIndex++)
+			{
+				const float Time = FrameRate.AsSeconds(FrameTimes[FrameIndex]) - TimeOffset;
+				const float Value = CurveValuesPerFrame[FrameIndex][CurveIndex];
+
+				Controller.SetCurveKey(CurveId, FRichCurveKey(Time, Value), false);
+			}
+		}
+	}
+	
+	Controller.NotifyPopulated();
+	Controller.CloseBracket();	
+	
+	Target->MarkPackageDirty();
+
+	return true;
+}
+
+bool URigMapperEditorSubsystem::LoadCurveValuesFromAnimSequence(const UAnimSequence* Source, TArray<FString>& CurveNames, TArray<FFrameTime>& FrameTimes, TArray<FPoseValues>& CurveValuesPerFrame)
+{
+	TArray<FPose> Poses;
+	
+	const TArray<FFloatCurve>& FloatCurves = Source->GetDataModel()->GetFloatCurves();
+
+	if (!FloatCurves.IsEmpty())
+	{
+		FrameTimes.Reserve(FloatCurves[0].FloatCurve.GetNumKeys());
+
+		for (const FFloatCurve& Curve : FloatCurves)
+		{
+			const FString CurveName = Curve.GetName().ToString();
+			// todo: remap name from inputs
+
+			if (!Curve.FloatCurve.Keys.IsEmpty())
+			{
+				CurveNames.AddUnique(CurveName);
+				
+				for (const FRichCurveKey& Key : Curve.FloatCurve.Keys)
+				{
+					const FFrameTime FrameNumber = GetAnimSequenceRate(Source).AsFrameTime(Key.Time);
+					int32 PoseIndex = FrameTimes.Find(FrameNumber);
+
+					if (PoseIndex == INDEX_NONE)
+					{
+						PoseIndex = FrameTimes.Add(FrameNumber);
+						Poses.AddDefaulted_GetRef().Reserve(FloatCurves.Num());
+					}
+
+					Poses[PoseIndex].Add(CurveName, Key.Value);
+				}
+			}
+		}
+	}
+
+	return BakeSparseKeys(Poses, CurveNames, FrameTimes, CurveValuesPerFrame);
+}
+
+bool URigMapperEditorSubsystem::ConvertAnimSequence(const UAnimSequence* Source, UAnimSequence* Target, const TArray<URigMapperDefinition*>& Definitions)
+{
+	TArray<FString> CurveNames;
+	TArray<FFrameTime> FrameTimes;
+	TArray<FPoseValues> CurveValuesPerFrame;
+	
+	if (!LoadCurveValuesFromAnimSequence(Source, CurveNames, FrameTimes, CurveValuesPerFrame))
+	{
+		return false;
+	}
+	
+	return ConvertCurveValuesToAnimSequence(CurveNames, FrameTimes, CurveValuesPerFrame, Target, Definitions);
+}
+
+bool URigMapperEditorSubsystem::ConvertAnimSequenceNew(const UAnimSequence* Source, USkeletalMesh* TargetMesh, const TArray<URigMapperDefinition*>& Definitions)
+{
+	UAnimSequence* NewSequence = CreateNewAnimSequence(Source->GetName(), FPaths::GetPath(Source->GetPathName()), TargetMesh);
+
+	NewSequence->ImportFileFramerate = Source->ImportFileFramerate;
+	NewSequence->ImportResampleFramerate = Source->ImportResampleFramerate;
+	SetAnimSequenceRate(NewSequence, GetAnimSequenceRate(Source), false);
+	
+	if (!ConvertAnimSequence(Source, NewSequence, Definitions))
+	{
+		return false;
+	}
+	NewSequence->PostLoad();
+	return true;
+}
+
+bool URigMapperEditorSubsystem::ConvertCsvToAnimSequence(const FFilePath& InputFile, UAnimSequence* Target, const TArray<URigMapperDefinition*>& Definitions)
+{
+	TArray<FString> CurveNames;
+	TArray<FFrameTime> FrameTimes;
+	TArray<FPoseValues> CurveValuesPerFrame;
+	
+	if (!LoadCurveValuesFromCsv(InputFile, CurveNames, FrameTimes, CurveValuesPerFrame))
+	{
+		return false;
+	}
+
+	return ConvertCurveValuesToAnimSequence(CurveNames, FrameTimes, CurveValuesPerFrame, Target, Definitions);
+}
+
+bool URigMapperEditorSubsystem::ConvertCsvToAnimSequenceNew(const FFilePath& InputFile, const FFilePath& NewAssetPath, USkeletalMesh* TargetMesh, const TArray<URigMapperDefinition*>& Definitions, const FFrameRate& FrameRate)
+{
+	UAnimSequence* NewSequence = CreateNewAnimSequence(FPaths::GetBaseFilename(InputFile.FilePath), NewAssetPath.FilePath, TargetMesh);
+	
+	SetAnimSequenceRate(NewSequence, FrameRate, true);
+
+	if (!ConvertCsvToAnimSequence(InputFile, NewSequence, Definitions))
+	{
+		return false;
+	}
+	NewSequence->PostLoad();
+	return true;
+}
+
+bool URigMapperEditorSubsystem::ConvertAnimSequenceToCsv(const UAnimSequence* Source, const FFilePath& OutputFile, const TArray<URigMapperDefinition*>& Definitions, bool bOutputIntermediateCsvFiles)
+{
+	TArray<FString> CurveNames;
+	TArray<FFrameTime> FrameTimes;
+	TArray<FPoseValues> CurveValuesPerFrame;
+	
+	if (!LoadCurveValuesFromAnimSequence(Source, CurveNames, FrameTimes, CurveValuesPerFrame))
+	{
+		return false;
+	}
+
+	return ConvertCurveValuesToCsv(CurveNames, FrameTimes, CurveValuesPerFrame, OutputFile, Definitions, bOutputIntermediateCsvFiles);
+}
+
+bool URigMapperEditorSubsystem::ConvertAnimSequenceToControlRigSection(const UAnimSequence* Source, UMovieSceneControlRigParameterSection* Target, const TArray<URigMapperDefinition*>& Definitions)
+{
+	TArray<FString> CurveNames;
+	TArray<FFrameTime> FrameTimes;
+	TArray<FPoseValues> CurveValuesPerFrame;
+	
+	if (!LoadCurveValuesFromAnimSequence(Source, CurveNames, FrameTimes, CurveValuesPerFrame))
+	{
+		return false;
+	}
+
+	return ConvertCurveValuesToControlRigSection(CurveNames, FrameTimes, CurveValuesPerFrame, Target, Definitions);
+}
+
+bool URigMapperEditorSubsystem::ConvertAnimSequenceToControlRigSectionNew(const UAnimSequence* Source, USkeletalMesh* TargetMesh, const TArray<URigMapperDefinition*>& Definitions, const TSubclassOf<UControlRig>& ControlRigClass)
+{
+	UMovieSceneControlRigParameterSection* NewSection = CreateControlRigSequence(FPaths::GetPath(Source->GetPathName()), Source->GetName(), Source->GetTypedOuter<UMovieScene>()->GetDisplayRate(), ControlRigClass, TargetMesh);
+	
+	TArray<FString> CurveNames;
+	TArray<FFrameTime> FrameTimes;
+	TArray<FPoseValues> CurveValuesPerFrame;
+	
+	if (!LoadCurveValuesFromAnimSequence(Source, CurveNames, FrameTimes, CurveValuesPerFrame))
+	{
+		return false;
+	}
+	
+	return ConvertCurveValuesToControlRigSection(CurveNames, FrameTimes, CurveValuesPerFrame, NewSection, Definitions);
+}
+
+bool URigMapperEditorSubsystem::ConvertCurveValuesToControlRigSection(TArray<FString>& CurveNames, const TArray<FFrameTime>& FrameTimes, TArray<FPoseValues>& CurveValuesPerFrame, UMovieSceneControlRigParameterSection* Target, const TArray<URigMapperDefinition*>& Definitions)
+{
+	TArray<TSharedPtr<FacialRigMapping::FRigMapper>> RigMappers;
+	if (!InitializeRigMapping(Definitions, RigMappers))
+	{
+		// todo: warn
+		return false;
+	}
+	
+	for (int32 RigIndex = 0; RigIndex < RigMappers.Num(); RigIndex++)
+	{
+		EvaluateRigMapping(CurveNames, FrameTimes, CurveValuesPerFrame, RigMappers[RigIndex]);
+	}
+
+	return AddCurveValuesToControlRigSection(Target, CurveNames, FrameTimes, CurveValuesPerFrame);
+}
+
+int32 GetChannelIndexFromCurveName(const FString& CurveName, const TArrayView<const FMovieSceneChannelMetaData> MetaData)
+{
+	TArray<FString> Parts;
+	CurveName.ParseIntoArray(Parts, TEXT("."));
+	
+	const FString ControlName = Parts[0];
+
+	// todo: better mapping for custom attrs
+	// .tx/ry/sz -> .Location.X/Rotation.Y/Scale.Z
+	const FString ChannelTmName = ControlName + Parts.Last().ToUpper().Replace(
+		TEXT("T"), TEXT(".Location.")).Replace(
+			TEXT("R"), TEXT(".Rotation.")).Replace(
+				TEXT("S"), TEXT(".Scale."));
+	const FString Channel2dName = ChannelTmName.Replace(TEXT(".Location"), TEXT(""));
+
+	int32 ChannelIndex = INDEX_NONE;
+	for (int32 MetaDataIndex = 0; MetaDataIndex < MetaData.Num(); MetaDataIndex++)
+	{
+		const FString ChannelName = MetaData[MetaDataIndex].Name.ToString();
+		
+		if (ChannelName == ControlName || ChannelName == ChannelTmName || ChannelName == Channel2dName)
+		{
+			ChannelIndex = MetaDataIndex;
+			break;
+		}
+	}
+	
+	return ChannelIndex;
+}
+
+bool URigMapperEditorSubsystem::AddCurveValuesToControlRigSection(UMovieSceneControlRigParameterSection* Target, const TArray<FString>& CurveNames, const TArray<FFrameTime>& FrameTimes, const TArray<FPoseValues>& CurveValuesPerFrame)
+{
+	const UMovieScene* MovieScene = Target->GetTypedOuter<UMovieScene>();
+
+	check(MovieScene);
+	
+	if (CurveNames.IsEmpty() || FrameTimes.IsEmpty() || CurveValuesPerFrame.IsEmpty())
+	{
+		return false;
+	}
+
+	Target->Modify();
+	// Target->ClearAllParameters();
+	// Target->GetControlRig()->RequestInit();
+	//Target->GetControlRig()->Initialize();
+	UControlRig* ControlRig = Target->GetControlRig(); // NewObject<UControlRig>(Target->GetOuter(), ->GetClass(), *Target->GetControlRig()->GetName(), RF_Transactional);
+	ControlRig->Modify();
+	ControlRig->Initialize();
+	ControlRig->RequestInit();
+	ControlRig->Execute(FRigUnit_PrepareForExecution::EventName);
+	Target->RecreateWithThisControlRig(ControlRig, true);
+	// Target->ReconstructChannelProxy();
+	// Target->CacheChannelProxy();
+
+	if (const FMovieSceneChannelEntry* Entry = Target->GetChannelProxy().FindEntry(FMovieSceneFloatChannel::StaticStruct()->GetFName()))
+	{
+		const TArrayView<FMovieSceneChannel* const> Channels = Entry->GetChannels();
+		if (Channels.IsEmpty())
+		{
+			return false;
+		}
+		const TArrayView<const FMovieSceneChannelMetaData> MetaData = Entry->GetMetaData();
+		
+		for (int32 CurveIndex = 0; CurveIndex < CurveNames.Num(); CurveIndex++)
+		{
+			const int32 ChannelIndex = GetChannelIndexFromCurveName(CurveNames[CurveIndex], MetaData);
+			// todo: remap from outputs
+
+			if (Channels.IsValidIndex(ChannelIndex)) // todo: else log
+			{
+				FMovieSceneFloatChannel* Channel = static_cast<FMovieSceneFloatChannel*>(Channels[ChannelIndex]);
+
+				// todo: do we want to clear the section?
+				if (Channel->GetNumKeys() > 0)
+				{
+					Channel->DeleteKeysFrom(Channel->GetTimes()[0] - 1, false);
+				}
+				
+				for (int32 FrameIndex = 0; FrameIndex < FrameTimes.Num(); FrameIndex++)
+				{
+					const FFrameNumber FrameNumber = MovieScene->GetTickResolution().AsFrameNumber(MovieScene->GetDisplayRate().AsSeconds(FrameTimes[FrameIndex]));
+					Channel->AddCubicKey(FrameNumber, CurveValuesPerFrame[FrameIndex][CurveIndex]);
+				}
+			}
+		}
+	}
+	return true;
+}
+
+FString ChannelNameToCurveName(const FString& ChannelName, const TArray<FString>& CurveNames)
+{
+	if (CurveNames.Contains(ChannelName))
+	{
+		return ChannelName;
+	}
+	
+	TArray<FString> Parts;
+	ChannelName.ParseIntoArray(Parts, TEXT("."));
+	
+	const FString ControlName = Parts[0];
+	FString ChannelToCurve;
+
+	if (Parts.Last() == ControlName)
+	{
+		// todo: default
+		ChannelToCurve = ControlName + ".ty";
+	}
+	else
+	{
+		if (Parts.Num() > 1)
+		{
+			ChannelToCurve = ControlName + Parts[1].Replace(
+				TEXT("Location"), TEXT(".t")).Replace(
+					TEXT("Rotation"), TEXT(".r")).Replace(
+						TEXT("Scale"), TEXT(".s")).Replace(
+							TEXT("X"), TEXT(".t")).Replace(
+								TEXT("Y"), TEXT(".t"));
+			ChannelToCurve += Parts.Last().ToLower();
+		}
+		// todo: any other remapping here?
+	}
+	
+	if (CurveNames.Contains(ChannelToCurve))
+	{
+		return ChannelToCurve;
+	}
+
+	return ""; // todo: warn
+}
+
+bool URigMapperEditorSubsystem::LoadCurveValuesFromControlRigSection(const UMovieSceneControlRigParameterSection* Source, TArray<FString>& CurveNames, TArray<FFrameTime>& FrameTimes, TArray<FPoseValues>& CurveValuesPerFrame, const TArray<FString>& InputNames)
+{
+	TArray<FPose> Poses;
+	
+	const UMovieScene* MovieScene = Source->GetTypedOuter<UMovieScene>();
+	check(MovieScene);
+	
+	// todo: non float channels
+	if (const FMovieSceneChannelEntry* Entry = Source->GetChannelProxy().FindEntry(FMovieSceneFloatChannel::StaticStruct()->GetFName()))
+	{
+		const TArrayView<FMovieSceneChannel* const> Channels = Entry->GetChannels();
+		if (!Channels.IsEmpty())
+		{
+			const TArrayView<const FMovieSceneChannelMetaData> MetaData = Entry->GetMetaData();
+			Poses.Reserve(Channels[0]->GetNumKeys());
+		
+			for (int32 ChannelIndex = 0; ChannelIndex < Channels.Num(); ChannelIndex++)
+			{
+				const FString ChannelName = MetaData[ChannelIndex].Name.ToString();
+				const FString CurveName = ChannelNameToCurveName(ChannelName, InputNames);
+				// todo: remap name from inputs
+
+				if (!CurveName.IsEmpty())
+				{
+					CurveNames.AddUnique(CurveName);
+				
+					const FMovieSceneFloatChannel* Channel = static_cast<FMovieSceneFloatChannel*>(Channels[ChannelIndex]);
+					TArrayView<const FFrameNumber> KeyTimes = Channel->GetTimes();
+					TArrayView<const FMovieSceneFloatValue> KeyValues = Channel->GetValues();
+
+					for (int32 KeyIndex = 0; KeyIndex < KeyTimes.Num() && KeyIndex < KeyValues.Num(); KeyIndex++)
+					{
+						const FFrameTime FrameNumber = MovieScene->GetDisplayRate().AsFrameTime(MovieScene->GetTickResolution().AsSeconds(KeyTimes[KeyIndex]));
+						int32 PoseIndex = FrameTimes.Find(FrameNumber);
+				
+						if (PoseIndex == INDEX_NONE)
+						{
+							PoseIndex = FrameTimes.Add(FrameNumber);
+							Poses.AddDefaulted_GetRef().Reserve(Channels.Num());
+						}
+				
+						Poses[PoseIndex].Add(CurveName, KeyValues[KeyIndex].Value);	
+					}
+				}
+			}
+		}
+	}
+	
+	return BakeSparseKeys(Poses, CurveNames, FrameTimes, CurveValuesPerFrame);
+}
+
+bool URigMapperEditorSubsystem::ConvertControlRigSection(const UMovieSceneControlRigParameterSection* Source, UMovieSceneControlRigParameterSection* Target, const TArray<URigMapperDefinition*>& Definitions)
+{
+	TArray<FString> CurveNames;
+	TArray<FFrameTime> FrameTimes;
+	TArray<FPoseValues> CurveValuesPerFrame;
+
+	if (Definitions.IsEmpty())
+	{
+		return false;
+	}
+	if (!LoadCurveValuesFromControlRigSection(Source, CurveNames, FrameTimes, CurveValuesPerFrame, Definitions[0]->Inputs))
+	{
+		return false;
+	}
+	
+	return ConvertCurveValuesToControlRigSection(CurveNames, FrameTimes, CurveValuesPerFrame, Target, Definitions);
+}
+
+bool URigMapperEditorSubsystem::ConvertControlRigSectionNew(const UMovieSceneControlRigParameterSection* Source, USkeletalMesh* TargetMesh, const TArray<URigMapperDefinition*>& Definitions, const TSubclassOf<UControlRig>& ControlRigClass)
+{
+	UMovieSceneControlRigParameterSection* NewSection = CreateControlRigSequence(FPaths::GetPath(Source->GetPathName()), Source->GetName(), Source->GetTypedOuter<UMovieScene>()->GetDisplayRate(), ControlRigClass, TargetMesh);
+	
+	return ConvertControlRigSection(Source, NewSection, Definitions);
+}
+
+bool URigMapperEditorSubsystem::ConvertCsvToControlRigSection(const FFilePath& InputFile, UMovieSceneControlRigParameterSection* Target, const TArray<URigMapperDefinition*>& Definitions)
+{
+	TArray<FString> CurveNames;
+	TArray<FFrameTime> FrameTimes;
+	TArray<FPoseValues> CurveValuesPerFrame;
+	
+	if (!LoadCurveValuesFromCsv(InputFile, CurveNames, FrameTimes, CurveValuesPerFrame))
+	{
+		return false;
+	}
+
+	return ConvertCurveValuesToControlRigSection(CurveNames, FrameTimes, CurveValuesPerFrame, Target, Definitions);
+}
+
+bool URigMapperEditorSubsystem::ConvertCsvToControlRigSectionNew(const FFilePath& InputFile, USkeletalMesh* TargetMesh, const TArray<URigMapperDefinition*>& Definitions, const FFilePath& NewAssetPath, const FFrameRate& FrameRate, const TSubclassOf<UControlRig>& ControlRigClass)
+{
+	UMovieSceneControlRigParameterSection* NewSection = CreateControlRigSequence(FPaths::GetPath(NewAssetPath.FilePath), FPaths::GetBaseFilename(InputFile.FilePath), FrameRate, ControlRigClass, TargetMesh);
+
+	return ConvertCsvToControlRigSection(InputFile, NewSection, Definitions);
+}
+
+bool URigMapperEditorSubsystem::ConvertControlRigSectionToCsv(const UMovieSceneControlRigParameterSection* Source, const FFilePath& OutputFile, const TArray<URigMapperDefinition*>& Definitions, bool bOutputIntermediateCsvFiles)
+{
+	TArray<FString> CurveNames;
+	TArray<FFrameTime> FrameTimes;
+	TArray<FPoseValues> CurveValuesPerFrame;
+
+	if (Definitions.IsEmpty())
+	{
+		return false;
+	}
+	if (!LoadCurveValuesFromControlRigSection(Source, CurveNames, FrameTimes, CurveValuesPerFrame, Definitions[0]->Inputs))
+	{
+		return false;
+	}
+	
+	return ConvertCurveValuesToCsv(CurveNames, FrameTimes, CurveValuesPerFrame, OutputFile, Definitions, bOutputIntermediateCsvFiles);
+}
+
+bool URigMapperEditorSubsystem::ConvertControlRigSectionToAnimSequence(const UMovieSceneControlRigParameterSection* Source, UAnimSequence* Target, const TArray<URigMapperDefinition*>& Definitions)
+{
+	TArray<FString> CurveNames;
+	TArray<FFrameTime> FrameTimes;
+	TArray<FPoseValues> CurveValuesPerFrame;
+
+	if (Definitions.IsEmpty())
+	{
+		return false;
+	}
+	if (!LoadCurveValuesFromControlRigSection(Source, CurveNames, FrameTimes, CurveValuesPerFrame, Definitions[0]->Inputs))
+	{
+		return false;
+	}
+
+	return ConvertCurveValuesToAnimSequence(CurveNames, FrameTimes, CurveValuesPerFrame, Target, Definitions);
+}
+
+bool URigMapperEditorSubsystem::ConvertControlRigSectionToAnimSequenceNew(const UMovieSceneControlRigParameterSection* Source, USkeletalMesh* TargetMesh, const TArray<URigMapperDefinition*>& Definitions)
+{
+	ULevelSequence* Sequence = GetSequenceFromSection(Source);
+	if (!Sequence || !Sequence->GetMovieScene())
+	{
+		return false;
+	}
+	
+	UAnimSequence* NewSequence = CreateNewAnimSequence(Sequence->GetName(), FPaths::GetPath(Sequence->GetPathName()), TargetMesh);
+	SetAnimSequenceRate(NewSequence, Sequence->GetMovieScene()->GetDisplayRate(), true);
+	
+	TArray<FString> CurveNames;
+	TArray<FFrameTime> FrameTimes;
+	TArray<FPoseValues> CurveValuesPerFrame;
+
+	if (Definitions.IsEmpty())
+	{
+		return false;
+	}
+	if (!LoadCurveValuesFromControlRigSection(Source, CurveNames, FrameTimes, CurveValuesPerFrame, Definitions[0]->Inputs))
+	{
+		return false;
+	}
+
+	if (!ConvertCurveValuesToAnimSequence(CurveNames, FrameTimes, CurveValuesPerFrame, NewSequence, Definitions))
+	{
+		return false;
+	}
+	NewSequence->PostLoad();
+	return true;
+}
+
+bool URigMapperEditorSubsystem::LoadCurveValuesFromCsv(const FFilePath& InputFile, TArray<FString>& CurveNames, TArray<FFrameTime>& FrameTimes, TArray<FPoseValues>& CurveValuesPerFrame)
+{
+	TArray<FPose> Poses;
+	
+	TArray<FString> CsvLines;
+	
+	if (!FFileHelper::LoadFileToStringArray(CsvLines, *InputFile.FilePath))
+	{
+		return false;
+	}
+
+	CsvLines.RemoveAt(0);
+	
+	int32 PoseIndex = INDEX_NONE;
+	FFrameTime LastFrameNumber(0);
+	FString LastCurveName;
+
+	int32 NumberCurves = 1;
+
+	// todo: mem
+	for (const FString& Line : CsvLines)
+	{
+		TArray<FString> Parts;
+
+		// todo: namespace
+		if (Line.Replace(TEXT(" "), TEXT("")).ParseIntoArray(Parts, TEXT(",")) != 3)
+		{
+			// todo: warn
+			continue;
+		}
+		const FString& CurveName = Parts[0];
+		const FFrameNumber FrameNumber = FCString::Atoi(*Parts[1]);
+		const float CurveValue = FCString::Atof(*Parts[2]);
+
+		if (CurveName != LastCurveName)
+		{
+			CurveNames.AddUnique(CurveName);
+			LastCurveName = CurveName;
+		}
+
+		if (FrameNumber != LastFrameNumber || PoseIndex == INDEX_NONE)
+		{
+			PoseIndex = FrameTimes.Find(FFrameTime(FrameNumber));
+			LastFrameNumber = FrameNumber;
+		}
+
+		if (PoseIndex == INDEX_NONE)
+		{
+			PoseIndex = FrameTimes.Add(FrameNumber);
+			Poses.AddDefaulted_GetRef().Reserve(NumberCurves);
+		}
+		
+		Poses[PoseIndex].Add(CurveName, CurveValue);
+		
+		if (Poses[PoseIndex].Num() > NumberCurves)
+		{
+			NumberCurves = Poses[PoseIndex].Num();
+		}
+	}
+	
+	return BakeSparseKeys(Poses, CurveNames, FrameTimes, CurveValuesPerFrame);
+}
+
+bool URigMapperEditorSubsystem::InitializeRigMapping(const TArray<URigMapperDefinition*>& Definitions, TArray<TSharedPtr<FacialRigMapping::FRigMapper>>& RigMappers)
+{
+	if (Definitions.IsEmpty())
+	{
+		UE_LOG(LogTemp, Error, TEXT("Not enough definitions to remap"))
+		return false;
+	}
+	
+	RigMappers.Empty(Definitions.Num());
+	
+	for (int i = 0; i < Definitions.Num(); i++)
+	{
+		if (!Definitions[i])
+		{
+			UE_LOG(LogTemp, Error, TEXT("Invalid definition at index: %d"), i)
+			RigMappers.Empty();
+			return false;
+		}
+		TSharedPtr<FacialRigMapping::FRigMapper> RigMapper = MakeShared<FacialRigMapping::FRigMapper>();
+		
+		if (!RigMapper->LoadDefinition(Definitions[i]))
+		{
+			UE_LOG(LogTemp, Error, TEXT("Could not load definition at index %d"), i)
+			RigMappers.Empty();
+			return false;
+		}
+		
+		RigMappers.Add(RigMapper);
+	}
+
+	// todo: check definitions chain 
+	
+	return true;
+}
+
+void URigMapperEditorSubsystem::EvaluateRigMapping(TArray<FString>& CurveNames, const TArray<FFrameTime>& FrameTimes, TArray<FPoseValues>& CurveValuesPerFrame, const TSharedPtr<FacialRigMapping::FRigMapper>& RigMapper)
+{
+	TArray<FString> NewCurveNames;
+	
+	for (int32 FrameIndex = 0;  FrameIndex < FrameTimes.Num(); FrameIndex++)
+	{
+		RigMapper->SetDirty();
+		
+		const TArray<FString>& InputNames = RigMapper->GetInputNames();
+		
+		for (int32 CurveIndex = 0; CurveIndex < CurveNames.Num(); CurveIndex++)
+		{
+			if (const int32 InputIndex = InputNames.Find(CurveNames[CurveIndex]); InputIndex != INDEX_NONE)
+			{
+				RigMapper->SetDirectValue(InputIndex, CurveValuesPerFrame[FrameIndex][CurveIndex]);
+			}
+			else
+			{
+				// todo: log / interp -> missing input?
+				UE_LOG(LogTemp, Log, TEXT("Missing key, interp not avaialble yet"))
+			}
+		}
+
+		TMap<FString, double> Outputs = RigMapper->GetOutputValues(true);
+
+		if (NewCurveNames.IsEmpty())
+		{
+			NewCurveNames.Reserve(Outputs.Num());
+		}
+		CurveValuesPerFrame[FrameIndex].Empty(Outputs.Num());
+		CurveValuesPerFrame[FrameIndex].AddDefaulted(Outputs.Num());
+		
+		for (const TPair<FString, double>& Output : Outputs)
+		{
+			const int32 CurveIndex = NewCurveNames.AddUnique(Output.Key);
+			CurveValuesPerFrame[FrameIndex][CurveIndex] = Output.Value;;
+		}			
+	}
+	CurveNames = NewCurveNames;
+}
+
+bool URigMapperEditorSubsystem::WriteCurveValuesToCsv(const FFilePath& OutputFile, const TArray<FString>& CurveNames, const TArray<FFrameTime>& FrameTimes, const TArray<FPoseValues>& CurveValuesPerFrame)
+{
+	TArray<FString> CsvLines;
+	CsvLines.Add(TEXT("curve_name, frame_number, value"));
+	
+	for (int32 FrameIndex = 0; FrameIndex < FrameTimes.Num(); FrameIndex++)
+	{
+		for (int32 CurveIndex = 0; CurveIndex < CurveNames.Num(); CurveIndex++)
+		{
+			CsvLines.Add(FString::Printf(TEXT("%s, %d, %f"), *CurveNames[CurveIndex], FrameTimes[FrameIndex].RoundToFrame().Value, CurveValuesPerFrame[FrameIndex][CurveIndex]));
+		}
+	}
+	
+	return FFileHelper::SaveStringArrayToFile(CsvLines, *OutputFile.FilePath);
+}
+
+bool URigMapperEditorSubsystem::BakeSparseKeys(const TArray<FPose>& Poses, const TArray<FString>& CurveNames, TArray<FFrameTime>& FrameTimes, TArray<FPoseValues>& CurveValuesPerFrame)
+{
+	if (Poses.IsEmpty() || CurveNames.IsEmpty() || Poses.Num() != FrameTimes.Num())
+	{
+		return false;
+	}
+
+	// This extra work is needed to make sure frame times and thus resulting curve values are in order
+	TArray<FFrameTime> FrameTimesInOrder = FrameTimes;
+	FrameTimesInOrder.Sort();
+	
+	CurveValuesPerFrame.Reset(FrameTimesInOrder.Num());
+	for (int32 FrameIndex = 0; FrameIndex < FrameTimes.Num(); FrameIndex += 1)
+	{
+		CurveValuesPerFrame.AddDefaulted_GetRef().AddZeroed(CurveNames.Num());
+	}
+
+	TArray<FString> BakedCurves;
+	BakedCurves.Reserve(CurveNames.Num());
+
+	// If all curves are present on frame 0, we won't go further than one loop
+	for (int32 FrameIndex = 0; FrameIndex < FrameTimes.Num() && BakedCurves.Num() < CurveNames.Num(); FrameIndex += 1)
+	{
+		const int32 ActualFrameIndex = FrameTimes.Find(FrameTimesInOrder[FrameIndex]); // Lookup for non ordered frames
+		
+		for (const TPair<FString, float>& Keyframe : Poses[ActualFrameIndex])
+		{
+			const FString& CurveName = Keyframe.Key;
+			
+			if (BakedCurves.Contains(CurveName))
+			{
+				continue;	
+			}
+			BakedCurves.Add(CurveName);
+
+			const float CurveValue = Keyframe.Value;
+			const int32 CurveIndex = CurveNames.Find(CurveName);
+
+			SparseBakeCurve(CurveName, CurveIndex, CurveValue, CurveValuesPerFrame, FrameTimesInOrder, FrameIndex, Poses, FrameTimes, ActualFrameIndex);
+			
+			if (BakedCurves.Num() == CurveNames.Num())
+			{
+				break;
+			}
+		}
+	}
+
+	FrameTimes = FrameTimesInOrder;
+	
+	return CurveValuesPerFrame.Num() == FrameTimes.Num() && CurveValuesPerFrame[0].Num() == CurveNames.Num() && CurveValuesPerFrame.Last().Num() == CurveNames.Num();
+}
+
+void URigMapperEditorSubsystem::SparseBakeCurve(const FString& CurveName, int32 CurveIndex, float CurveValue,
+	TArray<FPoseValues>& CurveValuesPerFrame, const TArray<FFrameTime>& FrameTimesInOrder, int32 FrameIndex,
+	const TArray<FPose>& Poses, const TArray<FFrameTime>& FrameTimes, const int32 ActualFrameIndex)
+{
+	// We haven't seen this control so far. So keep it const on all previous frames as well as the current one
+	for (int32 BakeFrameIndex = 0; BakeFrameIndex <= FrameIndex; BakeFrameIndex += 1)
+	{
+		CurveValuesPerFrame[BakeFrameIndex][CurveIndex] = CurveValue;
+	}
+	
+	float LastKeyedValue = CurveValue;
+	int32 LastKeyedFrameIndex = FrameIndex;
+	int32 ActualLastKeyedFrameIndex = ActualFrameIndex; // Lookup for non ordered frames
+
+	// For following frames, add a key a lerp all previously missed keys if need be
+	for (int32 BakeFrameIndex = FrameIndex + 1; BakeFrameIndex < FrameTimes.Num(); BakeFrameIndex += 1)
+	{
+		const int32 ActualBakeFrameIndex = FrameTimes.Find(FrameTimesInOrder[BakeFrameIndex]); // Lookup for non ordered frames
+
+		// If the current frame has a key for the control, set it, and try to lerp all the keys we might have missed since last time we saw the control
+		if (Poses[ActualBakeFrameIndex].Contains(CurveName))
+		{
+			// Key this frame
+			CurveValuesPerFrame[BakeFrameIndex][CurveIndex] = Poses[ActualBakeFrameIndex][CurveName];
+			LastKeyedValue = CurveValuesPerFrame[BakeFrameIndex][CurveIndex];
+			
+			// If we have not seen this control on the previous frame, lerp all keys with missed
+			if (BakeFrameIndex > LastKeyedFrameIndex + 1)
+			{
+				// Lerp from the last key we found for this control to the current key, adding a lerp key for each frame
+				for (int32 LerpFrameIndex = LastKeyedFrameIndex + 1; LerpFrameIndex < BakeFrameIndex; LerpFrameIndex += 1)
+				{
+					const double LerpAlpha = (FrameTimesInOrder[LerpFrameIndex].AsDecimal() - FrameTimesInOrder[LastKeyedFrameIndex].AsDecimal()) / (FrameTimesInOrder[BakeFrameIndex].AsDecimal() - FrameTimesInOrder[LastKeyedFrameIndex].AsDecimal());
+					const double LerpValue = FMath::Lerp(Poses[ActualLastKeyedFrameIndex][CurveName], Poses[ActualBakeFrameIndex][CurveName], LerpAlpha);
+
+					CurveValuesPerFrame[LerpFrameIndex][CurveIndex] = LerpValue;
+				}	
+			}
+			LastKeyedFrameIndex = BakeFrameIndex;
+			ActualLastKeyedFrameIndex = ActualBakeFrameIndex; // Lookup for non ordered frames
+		}
+		else
+		{
+			// Per default, add a const key to all next frames (in case there isn't any key next).
+			// If it happens that there is one later, LastKeyedFrameIndex has not changed, and the above will lerp and override all const keys 
+			CurveValuesPerFrame[BakeFrameIndex][CurveIndex] = LastKeyedValue;
+		}
+	}
+}
\ No newline at end of file
diff --git a/Engine/Plugins/Experimental/Animation/RigMapper/Source/RigMapperEditor/Private/RigMapperGraph/RigMapperDefinitionEditorGraph.cpp b/Engine/Plugins/Experimental/Animation/RigMapper/Source/RigMapperEditor/Private/RigMapperGraph/RigMapperDefinitionEditorGraph.cpp
new file mode 100644
--- /dev/null
+++ b/Engine/Plugins/Experimental/Animation/RigMapper/Source/RigMapperEditor/Private/RigMapperGraph/RigMapperDefinitionEditorGraph.cpp
@@ -0,0 +1,344 @@
+// Copyright Epic Games, Inc. All Rights Reserved.
+
+#include "RigMapperDefinitionEditorGraph.h"
+
+#include "RigMapperDefinition.h"
+#include "RigMapperDefinitionEditorGraphNode.h"
+#include "RigMapperDefinitionEditorGraphSchema.h"
+
+URigMapperDefinitionEditorGraph::URigMapperDefinitionEditorGraph(const FObjectInitializer& ObjectInitializer) : Super(ObjectInitializer)
+{
+	Schema = URigMapperDefinitionEditorGraphSchema::StaticClass();
+}
+
+void URigMapperDefinitionEditorGraph::Initialize(URigMapperDefinition* InDefinition)
+{
+	WeakDefinition = InDefinition;
+}
+
+void URigMapperDefinitionEditorGraph::RebuildGraph()
+{
+	RemoveAllNodes();
+	ConstructNodes();
+	NotifyGraphChanged();
+	RequestRefreshLayout(true);
+}
+
+void URigMapperDefinitionEditorGraph::ConstructNodes()
+{
+	if (!WeakDefinition.IsValid())
+	{
+		return;
+	}
+	URigMapperDefinition* Definition = WeakDefinition.Get();
+	
+	for (const TPair<FString, FString>& Output : Definition->Outputs)
+	{
+		CreateGraphNodesRec(Definition, Output.Key, true);
+	}
+
+	// Generate nodes not related to any output
+	TArray<FString> FeatureNames;
+	Definition->Features.GetFeatureNames(FeatureNames);
+	for (const FString& Feature : FeatureNames)
+	{
+		if (!FeatureNodes.Contains(Feature))
+		{
+			CreateGraphNodesRec(Definition, Feature, false);
+		}
+	}
+	for (const FString& Input : Definition->Inputs)
+	{
+		if (!InputNodes.Contains(Input))
+		{
+			CreateGraphNode(Input, ERigMapperNodeType::Input);
+		}
+	}
+}
+
+TArray<URigMapperDefinitionEditorGraphNode*> URigMapperDefinitionEditorGraph::GetNodesByName(const TArray<FString>& Inputs, const TArray<FString>& Features, const TArray<FString>& Outputs) const
+{
+	TArray<URigMapperDefinitionEditorGraphNode*> OutNodes;
+	OutNodes.Reserve(Inputs.Num() + Features.Num() + Outputs.Num());
+
+	auto AddNodesByNameFromMap = [&OutNodes](const TArray<FString>& InArray, TMap<FString, URigMapperDefinitionEditorGraphNode*> InMap)
+	{
+		for (const FString& NodeName : InArray)
+		{
+			if (URigMapperDefinitionEditorGraphNode** Node = InMap.Find(NodeName))
+			{
+				OutNodes.Add(*Node);
+			}
+		}
+	};
+
+	AddNodesByNameFromMap(Inputs, InputNodes);
+	AddNodesByNameFromMap(Features, FeatureNodes);
+	AddNodesByNameFromMap(Outputs, OutputNodes);
+
+	return OutNodes;
+}
+
+URigMapperDefinitionEditorGraphNode* URigMapperDefinitionEditorGraph::CreateGraphNodesRec(URigMapperDefinition* Definition, const FString& NodeName, bool bIsOutputNode)
+{
+	// No need to do anything is node was already created. Output nodes should not be referenced by other nodes and should only get created once.
+	if (!bIsOutputNode)
+	{
+		URigMapperDefinitionEditorGraphNode** ExistingNode = InputNodes.Find(NodeName);
+		if (!ExistingNode)
+		{
+			ExistingNode = FeatureNodes.Find(NodeName);
+		}
+		if (ExistingNode)
+		{
+			return *ExistingNode;		
+		}
+	}
+
+	// Node was not created (referenced by a previously created node)
+	if (bIsOutputNode)
+	{
+		return CreateOutputNode(Definition, NodeName);
+	}
+	if (Definition->Inputs.Contains(NodeName))
+	{
+		return CreateGraphNode(NodeName, ERigMapperNodeType::Input); 
+	}
+	return CreateFeatureNode(Definition, NodeName);
+}
+
+URigMapperDefinitionEditorGraphNode* URigMapperDefinitionEditorGraph::CreateOutputNode(URigMapperDefinition* Definition, const FString& NodeName)
+{
+	if (const FString* LinkedInputName = Definition->Outputs.Find(NodeName))
+	{
+		URigMapperDefinitionEditorGraphNode* Node = CreateGraphNode(NodeName, ERigMapperNodeType::Output);
+		
+		if (URigMapperDefinitionEditorGraphNode* LinkedNode = CreateGraphNodesRec(Definition, *LinkedInputName, false))
+		{
+			LinkGraphNodes(LinkedNode, Node);
+		}
+
+		return Node;
+	}
+	return nullptr;
+}
+
+URigMapperDefinitionEditorGraphNode* URigMapperDefinitionEditorGraph::CreateFeatureNode(URigMapperDefinition* Definition, const FString& NodeName)
+{
+	ERigMapperFeatureType FeatureType;
+
+	if (const FRigMapperFeature* Feature = Definition->Features.Find(NodeName, FeatureType))
+	{
+		const ERigMapperNodeType NodeType = static_cast<ERigMapperNodeType>(static_cast<uint8>(FeatureType));
+		URigMapperDefinitionEditorGraphNode* Node = CreateGraphNode(NodeName, NodeType);
+			
+		TArray<FString> FeatureInputNames;
+		Feature->GetInputs(FeatureInputNames);
+	
+		for (const FString& FeatureInput : FeatureInputNames)
+		{
+			if (URigMapperDefinitionEditorGraphNode* LinkedNode = CreateGraphNodesRec(Definition, FeatureInput, false))
+			{
+				LinkGraphNodes(LinkedNode, Node);
+			}
+		}
+		
+		return Node;
+	}
+	return nullptr;
+}
+
+void URigMapperDefinitionEditorGraph::LinkGraphNodes(URigMapperDefinitionEditorGraphNode* InNode, URigMapperDefinitionEditorGraphNode* OutNode)
+{
+	if (InNode != OutNode)
+	{
+		UEdGraphPin* InPin = OutNode->CreateInputPin();
+		InPin->bHidden = false;
+
+		UEdGraphPin* OutPin = InNode->GetOutputPin();
+		if (!OutPin)
+		{
+			OutPin = InNode->CreateOutputPin();
+			OutPin->bHidden = false;
+		}
+		if (OutPin)
+		{
+			OutPin->MakeLinkTo(InPin);
+		}
+	}
+}
+		
+
+URigMapperDefinitionEditorGraphNode* URigMapperDefinitionEditorGraph::CreateGraphNode(const FString& NodeName, ERigMapperNodeType NodeType)
+{
+	if (!WeakDefinition.IsValid())
+	{
+		return nullptr;
+	}
+	
+	constexpr bool bSelectNewNode = false;
+	
+	FGraphNodeCreator<URigMapperDefinitionEditorGraphNode> GraphNodeCreator(*this);
+	URigMapperDefinitionEditorGraphNode* Node = GraphNodeCreator.CreateNode(bSelectNewNode);
+	GraphNodeCreator.Finalize();
+
+	Node->SetupNode(WeakDefinition.Get(), NodeName, NodeType);
+
+	if (NodeType == ERigMapperNodeType::Input)
+	{
+		InputNodes.Add(NodeName, Node);
+	}
+	else if (NodeType == ERigMapperNodeType::Output)
+	{
+		OutputNodes.Add(NodeName, Node);
+	}
+	else // todo: Add feature type
+	{
+		FeatureNodes.Add(NodeName, Node);
+	}
+	
+	return Node;
+}
+
+void URigMapperDefinitionEditorGraph::LayoutNodeRec(URigMapperDefinitionEditorGraphNode* InNode, double InputsWidth, double PosY, TArray<URigMapperDefinitionEditorGraphNode*>& LayedOutNodes) const
+{
+	const int32 NodeMarginX = 20;
+	const int32 NodeMarginY = 5;
+	
+	URigMapperDefinitionEditorGraphNode* LinkedNode = nullptr;
+	double SubGraphHeight = 0;
+	
+	for (UEdGraphPin* InPin : InNode->GetInputPins())
+	{
+		for (const UEdGraphPin* OutPin : InPin->LinkedTo)
+		{
+			LinkedNode = Cast<URigMapperDefinitionEditorGraphNode>(OutPin->GetOwningNode());
+			if (!LayedOutNodes.Contains(LinkedNode))
+			{
+				LayedOutNodes.Add(LinkedNode);
+
+				const double DesiredPosY =  PosY + SubGraphHeight;
+				LayoutNodeRec(LinkedNode, InputsWidth, DesiredPosY, LayedOutNodes);
+				SubGraphHeight += LinkedNode->GetDimensions().Y + LinkedNode->GetMargin().Y + (LinkedNode->NodePosY - DesiredPosY);
+			}
+
+			const double TargetPosX = LinkedNode->NodePosX + LinkedNode->GetDimensions().X + LinkedNode->GetMargin().X; 
+			if (TargetPosX > InNode->NodePosX)
+			{
+				InNode->NodePosX = TargetPosX;
+			}
+		}
+	}
+	
+	InNode->NodePosY = PosY;
+
+	const FVector2D& Dimensions = InNode->GetDimensions();
+	FVector2D Margin = { NodeMarginX, NodeMarginY };
+	
+	if (InNode->GetNodeType() == ERigMapperNodeType::Input)
+	{
+		InNode->NodePosX = 0;
+		Margin.X += InputsWidth - Dimensions.X;
+	}
+	else if (InNode->GetNodeType() == ERigMapperNodeType::Output)
+	{
+		Margin.X = 0;
+	}
+	else
+	{
+		InNode->NodePosX = FMath::Max(InNode->NodePosX, InputsWidth);
+	}
+	if (SubGraphHeight > Dimensions.Y + Margin.Y)
+	{
+		const double Offset = SubGraphHeight / 2 - (Dimensions.Y + Margin.Y) / 2;
+
+		Margin.Y += Offset;
+		InNode->NodePosY += Offset;
+	}
+
+	InNode->SetMargin(Margin);
+}
+
+void URigMapperDefinitionEditorGraph::LayoutNodes() const
+{
+	double InputsMaxWidth = 0;
+
+	const int32 InputMarginX = 50;
+	for (const TPair<FString, URigMapperDefinitionEditorGraphNode*>& InputNode : InputNodes)
+	{
+		const FVector2D& NodeDimensions = InputNode.Value->GetDimensions();
+
+		if (NodeDimensions.X > InputsMaxWidth)
+		{
+			InputsMaxWidth = NodeDimensions.X;
+		}
+	}
+	InputsMaxWidth += InputMarginX;
+
+	double MaxPosX = 0;
+	TArray<URigMapperDefinitionEditorGraphNode*> LayedOutNodes;
+
+	double PosY = 0;
+	const int32 SubGraphMarginY = 25;
+	
+	for (const TPair<FString, URigMapperDefinitionEditorGraphNode*>& Output : OutputNodes)
+	{
+		LayoutNodeRec(Output.Value, InputsMaxWidth, PosY, LayedOutNodes);
+		PosY = Output.Value->NodePosY + Output.Value->GetDimensions().Y + Output.Value->GetMargin().Y + SubGraphMarginY;
+		if (Output.Value->NodePosX > MaxPosX)
+		{
+			MaxPosX = Output.Value->NodePosX;
+		}
+	}
+	
+	const int32 OutputMarginX = 50;
+	for (const TPair<FString, URigMapperDefinitionEditorGraphNode*>& Output : OutputNodes)
+	{
+		Output.Value->NodePosX = MaxPosX + OutputMarginX;
+		// for (UEdGraphPin* InPin : Output.Value->GetInputPins())
+		// {
+		// 	for (const UEdGraphPin* OutPin : InPin->LinkedTo)
+		// 	{
+		// 		URigMapperDefinitionEditorGraphNode* LinkedNode = Cast<URigMapperDefinitionEditorGraphNode>(OutPin->GetOwningNode());
+		//
+		// 		if (!LinkedNode->GetInputPins().IsEmpty())
+		// 		{
+		// 			LinkedNode->NodePosX = Output.Value->NodePosX - (LinkedNode->GetDimensions().X + LinkedNode->GetMargin().X + OutputMarginX);
+		// 		}
+		// 	}
+		// }
+	}
+
+	// Layout nodes not related to any output
+	for (const TPair<FString, URigMapperDefinitionEditorGraphNode*>& Node : FeatureNodes)
+	{
+		if (!LayedOutNodes.Contains(Node.Value))
+		{
+			LayoutNodeRec(Node.Value, InputsMaxWidth, PosY, LayedOutNodes);
+			PosY = Node.Value->NodePosY + Node.Value->GetDimensions().Y + Node.Value->GetMargin().Y + SubGraphMarginY;
+		}
+	}
+	for (const TPair<FString, URigMapperDefinitionEditorGraphNode*>& Node : InputNodes)
+	{
+		if (!LayedOutNodes.Contains(Node.Value))
+		{
+			LayoutNodeRec(Node.Value, InputsMaxWidth, PosY, LayedOutNodes);
+			PosY = Node.Value->NodePosY + Node.Value->GetDimensions().Y + Node.Value->GetMargin().Y + SubGraphMarginY;
+		}
+	}
+}
+
+void URigMapperDefinitionEditorGraph::RemoveAllNodes()
+{
+	TArray<UEdGraphNode*> NodesToRemove = Nodes;
+	
+	for (int32 NodeIndex = 0; NodeIndex < NodesToRemove.Num(); ++NodeIndex)
+	{
+		RemoveNode(NodesToRemove[NodeIndex]);
+	}
+
+	Nodes.Reset();
+	InputNodes.Reset();
+	FeatureNodes.Reset();
+	OutputNodes.Reset();
+}
diff --git a/Engine/Plugins/Experimental/Animation/RigMapper/Source/RigMapperEditor/Private/RigMapperGraph/RigMapperDefinitionEditorGraph.h b/Engine/Plugins/Experimental/Animation/RigMapper/Source/RigMapperEditor/Private/RigMapperGraph/RigMapperDefinitionEditorGraph.h
new file mode 100644
--- /dev/null
+++ b/Engine/Plugins/Experimental/Animation/RigMapper/Source/RigMapperEditor/Private/RigMapperGraph/RigMapperDefinitionEditorGraph.h
@@ -0,0 +1,50 @@
+// Copyright Epic Games, Inc. All Rights Reserved.
+
+#pragma once
+
+#include "CoreMinimal.h"
+#include "EdGraph/EdGraph.h"
+#include "RigMapperDefinitionEditorGraph.generated.h"
+
+enum class ERigMapperNodeType : uint8;
+class URigMapperDefinitionEditorGraphNode;
+class URigMapperDefinition;
+
+/**
+ * 
+ */
+UCLASS()
+class RIGMAPPEREDITOR_API URigMapperDefinitionEditorGraph : public UEdGraph
+{
+	GENERATED_UCLASS_BODY()
+
+public:
+	void Initialize(URigMapperDefinition* InDefinition);
+	
+	void RebuildGraph();
+	void ConstructNodes();
+	void RemoveAllNodes();
+	void RequestRefreshLayout(bool bInRefreshLayout) { bRefreshLayout = bInRefreshLayout; }
+	bool NeedsRefreshLayout() const { return bRefreshLayout; }
+	void LayoutNodes() const;
+	
+	URigMapperDefinition* GetDefinition() const { return WeakDefinition.Get(); };
+
+	TArray<URigMapperDefinitionEditorGraphNode*> GetNodesByName(const TArray<FString>& Inputs, const TArray<FString>& Features, const TArray<FString>& Outputs) const;
+
+private:
+	URigMapperDefinitionEditorGraphNode* CreateGraphNodesRec(URigMapperDefinition* Definition, const FString& InNodeName, bool bIsOutputNode);
+	URigMapperDefinitionEditorGraphNode* CreateOutputNode(URigMapperDefinition* Definition, const FString& NodeName);
+	URigMapperDefinitionEditorGraphNode* CreateFeatureNode(URigMapperDefinition* Definition, const FString& NodeName);
+	static void LinkGraphNodes(URigMapperDefinitionEditorGraphNode* InNode, URigMapperDefinitionEditorGraphNode* OutNode);
+	URigMapperDefinitionEditorGraphNode* CreateGraphNode(const FString& NodeName, ERigMapperNodeType NodeType);
+	void LayoutNodeRec(URigMapperDefinitionEditorGraphNode* InNode, double InputsWidth, double PosY, TArray<URigMapperDefinitionEditorGraphNode*>& LayedOutNodes) const;
+	
+private:
+	TWeakObjectPtr<URigMapperDefinition> WeakDefinition;
+	TMap<FString, URigMapperDefinitionEditorGraphNode*> InputNodes;
+	TMap<FString, URigMapperDefinitionEditorGraphNode*> FeatureNodes;
+	TMap<FString, URigMapperDefinitionEditorGraphNode*> OutputNodes;
+
+	bool bRefreshLayout = false;
+};
diff --git a/Engine/Plugins/Experimental/Animation/RigMapper/Source/RigMapperEditor/Private/RigMapperGraph/RigMapperDefinitionEditorGraphNode.cpp b/Engine/Plugins/Experimental/Animation/RigMapper/Source/RigMapperEditor/Private/RigMapperGraph/RigMapperDefinitionEditorGraphNode.cpp
new file mode 100644
--- /dev/null
+++ b/Engine/Plugins/Experimental/Animation/RigMapper/Source/RigMapperEditor/Private/RigMapperGraph/RigMapperDefinitionEditorGraphNode.cpp
@@ -0,0 +1,147 @@
+// Copyright Epic Games, Inc. All Rights Reserved.
+
+
+#include "RigMapperDefinitionEditorGraphNode.h"
+
+#include "SRigMapperDefinitionGraphEditorNode.h"
+#include "RigMapperDefinition.h"
+#include "RigMapperDefinitionEditorGraph.h"
+#include "EdGraph/EdGraphPin.h"
+
+TSharedPtr<SGraphNode> URigMapperDefinitionEditorGraphNode::NodeFactory::CreateNode(UEdGraphNode* Node) const
+{
+	if (URigMapperDefinitionEditorGraphNode* ThisNode = Cast<URigMapperDefinitionEditorGraphNode>(Node))
+	{
+		if (URigMapperDefinitionEditorGraph* Graph = Cast<URigMapperDefinitionEditorGraph>(Node->GetGraph()))
+		{
+			Graph->RequestRefreshLayout(true);
+		}
+
+		TSharedRef<SGraphNode> GraphNode = SNew(SRigMapperDefinitionGraphEditorNode, ThisNode);
+		GraphNode->SlatePrepass();
+		ThisNode->SetDimensions(GraphNode->GetDesiredSize());
+		return GraphNode;
+	}
+	return nullptr;
+}
+
+FLinearColor URigMapperDefinitionEditorGraphNode::GetNodeBodyTintColor() const
+{
+	return GetNodeTitleColor();
+}
+
+FLinearColor URigMapperDefinitionEditorGraphNode::GetNodeTitleColor() const
+{
+	if (NodeType == ERigMapperNodeType::Input)
+	{
+		return FLinearColor(0.466, 0.969, 0.878);
+	}
+	if (NodeType == ERigMapperNodeType::Output)
+	{
+		return FLinearColor(0.466, 0.969, 0.525);
+	}
+	if (NodeType == ERigMapperNodeType::Multiply)
+	{
+		return FLinearColor(0.969, 0.966, 0.466);
+	}
+	if (NodeType == ERigMapperNodeType::SDK)
+	{
+		return FLinearColor(0.969, 0.717, 0.466);
+	}
+	if (NodeType == ERigMapperNodeType::WeightedSum)
+	{
+		return FLinearColor(0.65, 0.466, 0.969);
+	}
+	return FLinearColor(0.969, 0.5, 0.466);
+}
+
+void URigMapperDefinitionEditorGraphNode::SetupNode(URigMapperDefinition* InDefinition, const FString& InNodeName, ERigMapperNodeType InNodeType)
+{
+	NodeName = InNodeName;
+	Definition = InDefinition;
+	
+	NodeTitle = FText::FromString(NodeName);
+	SetNodeType(InNodeType);
+}
+
+void URigMapperDefinitionEditorGraphNode::SetNodeType(ERigMapperNodeType InNodeType)
+{
+	NodeType = InNodeType;
+
+	if (Definition.IsValid() && (NodeType == ERigMapperNodeType::WeightedSum || NodeType == ERigMapperNodeType::SDK))
+	{
+		ERigMapperFeatureType FeatureType;
+		if (FRigMapperFeature* Feature = Definition.Get()->Features.Find(NodeName, FeatureType))
+		{
+			if (FeatureType == ERigMapperFeatureType::SDK)
+			{
+				TArray<FString> InKeys;
+				TArray<FString> OutKeys;
+				
+				const TArray<FRigMapperSdkKey>& Keys = static_cast<FRigMapperSdkFeature*>(Feature)->Keys;
+
+				InKeys.Reserve(Keys.Num());
+				OutKeys.Reserve(Keys.Num());
+				for (const FRigMapperSdkKey& Key : Keys)
+				{
+					InKeys.Add(FString::Printf(TEXT("%.*f"), 3, Key.In));
+					OutKeys.Add(FString::Printf(TEXT("%.*f"), 3, Key.Out));
+				}
+				const FString In = FString::Join(InKeys, TEXT(", "));
+				const FString Out = FString::Join(OutKeys, TEXT(", "));
+				NodeSubtitle = FText::FromString(FString::Printf(TEXT("[%s] > [%s]"), *In, *Out));
+			}
+			if (FeatureType == ERigMapperFeatureType::WeightedSum)
+			{
+				const TMap<FString, double>& Inputs = static_cast<FRigMapperWsFeature*>(Feature)->Inputs;
+				const FRigMapperFeatureRange& Range = static_cast<FRigMapperWsFeature*>(Feature)->Range;
+
+				for (UEdGraphPin* InPin : InputPins)
+				{
+					for (const UEdGraphPin* OutPin : InPin->LinkedTo)
+					{
+						const FString LinkedNodeName = OutPin->GetOwningNode()->GetNodeTitle(ENodeTitleType::FullTitle).ToString();
+						if (Inputs.Contains(LinkedNodeName))
+						{
+							InPin->PinFriendlyName = FText::Format(INVTEXT("{0}"), Inputs[LinkedNodeName]);
+							FString RangeStr;
+							if (Range.bHasLowerBound)
+							{
+								RangeStr.Appendf(TEXT("min: %.*f"), 3, Range.LowerBound);
+							}
+							if (Range.bHasLowerBound && Range.bHasUpperBound)
+							{
+								RangeStr.Append(TEXT("\n"));
+							}
+							if (Range.bHasUpperBound)
+							{
+								RangeStr.Appendf(TEXT("max: %.*f"), 3, Range.UpperBound);
+							}
+							NodeSubtitle = FText::FromString(RangeStr);
+						}	
+					}
+				}
+			}
+		}
+		
+	}
+
+}
+
+UEdGraphPin* URigMapperDefinitionEditorGraphNode::CreateInputPin()
+{
+	return InputPins.Add_GetRef(CreatePin(EEdGraphPinDirection::EGPD_Input, NAME_None, NAME_None, NAME_None));
+}
+
+UEdGraphPin* URigMapperDefinitionEditorGraphNode::CreateOutputPin()
+{
+	return OutputPins.Add_GetRef(CreatePin(EEdGraphPinDirection::EGPD_Output, NAME_None, NAME_None, NAME_None));
+}
+
+void URigMapperDefinitionEditorGraphNode::GetRect(FVector2D& TopLeft, FVector2D& BottomRight) const
+{
+	TopLeft.X = NodePosX;
+	TopLeft.Y = NodePosY;
+	BottomRight.X = NodePosX + GetDimensions().X + GetMargin().X;
+	BottomRight.Y = NodePosY + GetDimensions().Y + GetMargin().Y;
+}
diff --git a/Engine/Plugins/Experimental/Animation/RigMapper/Source/RigMapperEditor/Private/RigMapperGraph/RigMapperDefinitionEditorGraphNode.h b/Engine/Plugins/Experimental/Animation/RigMapper/Source/RigMapperEditor/Private/RigMapperGraph/RigMapperDefinitionEditorGraphNode.h
new file mode 100644
--- /dev/null
+++ b/Engine/Plugins/Experimental/Animation/RigMapper/Source/RigMapperEditor/Private/RigMapperGraph/RigMapperDefinitionEditorGraphNode.h
@@ -0,0 +1,90 @@
+// Copyright Epic Games, Inc. All Rights Reserved.
+
+#pragma once
+
+#include "CoreMinimal.h"
+#include "EdGraphUtilities.h"
+#include "EdGraph/EdGraphNode.h"
+
+#include "RigMapperDefinitionEditorGraphNode.generated.h"
+
+class URigMapperDefinition;
+
+UENUM()
+enum class ERigMapperNodeType : uint8
+{
+	Input,
+	WeightedSum,
+	SDK,
+	Multiply,
+	Output,
+	Invalid
+};
+
+/**
+ * 
+ */
+UCLASS()
+class RIGMAPPEREDITOR_API URigMapperDefinitionEditorGraphNode : public UEdGraphNode
+{
+	GENERATED_BODY()
+
+public:
+	class NodeFactory : public FGraphPanelNodeFactory
+	{
+		virtual TSharedPtr<SGraphNode> CreateNode(UEdGraphNode* Node) const override;
+	};
+
+	
+	// UEdGraphNode implementation
+	virtual FText GetNodeTitle(ENodeTitleType::Type TitleType) const override { return NodeTitle; }
+	virtual FLinearColor GetNodeBodyTintColor() const override;
+	virtual FLinearColor GetNodeTitleColor() const override;
+	// !UEdGraphNode implementation
+
+	void SetupNode(URigMapperDefinition* InDefinition, const FString& InNodeName, ERigMapperNodeType InNodeType);
+	const FString& GetNodeName() const { return NodeName; }
+	
+	void SetNodeType(ERigMapperNodeType InNodeType);
+	ERigMapperNodeType GetNodeType() const { return NodeType; }
+	
+	const FText& GetSubtitle() const { return NodeSubtitle; }
+	
+	UEdGraphPin* CreateInputPin();
+	UEdGraphPin* CreateOutputPin();
+	const TArray<UEdGraphPin*>& GetInputPins() const { return InputPins; }
+	UEdGraphPin* GetOutputPin() const { return OutputPins.IsEmpty() ? nullptr : OutputPins[0]; };
+
+	void SetDimensions(const FVector2D& InDimensions) { Dimensions = InDimensions; }
+	const FVector2D& GetDimensions() const { return Dimensions; }
+
+	void SetMargin(const FVector2D& InDMargin) { Margin = InDMargin; }
+	const FVector2D& GetMargin() const { return Margin; }
+
+	void GetRect(FVector2D& TopLeft, FVector2D& BottomRight) const;
+
+protected:
+	/** Cached title for the node */
+	FText NodeTitle;
+
+	/** Cached title for the node */
+	FText NodeSubtitle;
+
+	/** Our one input pin */
+	TArray<UEdGraphPin*> InputPins;
+
+	/** Our one output pin */
+	TArray<UEdGraphPin*> OutputPins;
+
+	/** Cached dimensions of this node (used for layout) */
+	FVector2D Dimensions = { 300, 50 };
+
+	/** Cached dimensions of this node (used for layout) */
+	FVector2D Margin = { 0, 0 };
+	
+	TWeakObjectPtr<URigMapperDefinition> Definition;
+
+	FString NodeName;
+
+	ERigMapperNodeType NodeType = ERigMapperNodeType::Invalid;
+};
diff --git a/Engine/Plugins/Experimental/Animation/RigMapper/Source/RigMapperEditor/Private/RigMapperGraph/RigMapperDefinitionEditorGraphSchema.cpp b/Engine/Plugins/Experimental/Animation/RigMapper/Source/RigMapperEditor/Private/RigMapperGraph/RigMapperDefinitionEditorGraphSchema.cpp
new file mode 100644
--- /dev/null
+++ b/Engine/Plugins/Experimental/Animation/RigMapper/Source/RigMapperEditor/Private/RigMapperGraph/RigMapperDefinitionEditorGraphSchema.cpp
@@ -0,0 +1,98 @@
+// Copyright Epic Games, Inc. All Rights Reserved.
+
+
+#include "RigMapperDefinitionEditorGraphSchema.h"
+
+#include "RigMapperDefinition.h"
+#include "RigMapperDefinitionEditorGraphNode.h"
+#include "SRigMapperDefinitionGraphEditorNode.h"
+
+#include "EdGraph/EdGraph.h"
+
+
+#define LOCTEXT_NAMESPACE "RigMapperDefinitionEditorGraphSchema"
+
+
+void URigMapperDefinitionEditorGraphSchema::GetContextMenuActions(UToolMenu* Menu, UGraphNodeContextMenuContext* Context) const
+{
+	// Menu->AddDynamicSection("PhysicsAssetGraphSchema", FNewToolMenuDelegateLegacy::CreateLambda([](FMenuBuilder& MenuBuilder, UToolMenu* InMenu)
+	// {
+	// 	UGraphNodeContextMenuContext* ContextObject = InMenu->FindContext<UGraphNodeContextMenuContext>();
+	// 	if (!ContextObject)
+	// 	{
+	// 		return;
+	// 	}
+	//
+	// 	const UPhysicsAssetGraph* PhysicsAssetGraph = CastChecked<const UPhysicsAssetGraph>(ContextObject->Graph);
+	// 	TSharedPtr<FPhysicsAssetEditorSharedData> SharedData = PhysicsAssetGraph->GetPhysicsAssetEditor()->GetSharedData();
+	//
+	// 	if (const UPhysicsAssetGraphNode_Constraint* ConstraintNode = Cast<const UPhysicsAssetGraphNode_Constraint>(ContextObject->Node))
+	// 	{
+	// 		PhysicsAssetGraph->GetPhysicsAssetEditor()->BuildMenuWidgetConstraint(MenuBuilder);
+	// 	}
+	// 	else if (const UPhysicsAssetGraphNode_Bone* BoneNode = Cast<const UPhysicsAssetGraphNode_Bone>(ContextObject->Node))
+	// 	{
+	// 		PhysicsAssetGraph->GetPhysicsAssetEditor()->BuildMenuWidgetBody(MenuBuilder);
+	// 	}
+	//
+	// 	PhysicsAssetGraph->GetPhysicsAssetEditor()->BuildMenuWidgetSelection(MenuBuilder);
+	// }));
+}
+
+FLinearColor URigMapperDefinitionEditorGraphSchema::GetPinTypeColor(const FEdGraphPinType& PinType) const
+{
+	return FLinearColor::White;
+}
+
+void URigMapperDefinitionEditorGraphSchema::BreakPinLinks(UEdGraphPin& TargetPin, bool bSendsNodeNotification) const
+{
+	// Don't allow breaking any links
+}
+
+void URigMapperDefinitionEditorGraphSchema::BreakSinglePinLink(UEdGraphPin* SourcePin, UEdGraphPin* TargetPin) const
+{
+	// Don't allow breaking any links
+}
+
+FPinConnectionResponse URigMapperDefinitionEditorGraphSchema::MovePinLinks(UEdGraphPin& MoveFromPin, UEdGraphPin& MoveToPin, bool bIsIntermediateMove, bool bNotifyLinkedNodes) const
+{
+	// Don't allow moving any links
+	return FPinConnectionResponse(CONNECT_RESPONSE_DISALLOW, TEXT(""));
+}
+
+FPinConnectionResponse URigMapperDefinitionEditorGraphSchema::CopyPinLinks(UEdGraphPin& CopyFromPin, UEdGraphPin& CopyToPin, bool bIsIntermediateCopy) const
+{
+	// Don't allow copying any links
+	return FPinConnectionResponse(CONNECT_RESPONSE_DISALLOW, TEXT(""));
+}
+
+FConnectionDrawingPolicy* URigMapperDefinitionEditorGraphSchema::CreateConnectionDrawingPolicy(int32 InBackLayerID, int32 InFrontLayerID, float InZoomFactor, const FSlateRect& InClippingRect, class FSlateWindowElementList& InDrawElements, class UEdGraph* InGraphObj) const
+{
+	return new ConnectionDrawingPolicy(InBackLayerID, InFrontLayerID, InZoomFactor, InClippingRect, InDrawElements);
+}
+
+const FPinConnectionResponse URigMapperDefinitionEditorGraphSchema::CanCreateConnection(const UEdGraphPin* A, const UEdGraphPin* B) const
+{
+	return FPinConnectionResponse(CONNECT_RESPONSE_DISALLOW, LOCTEXT("RigMapperDefinitionEditor_CannotCreateConnection", "It is currently not possible to create new connections from the graph"));
+}
+
+FPinConnectionResponse URigMapperDefinitionEditorGraphSchema::CanCreateNewNodes(UEdGraphPin* InSourcePin) const
+{
+	// if(URigMapperDefinitionEditorGraphNode* Node = Cast<URigMapperDefinitionEditorGraphNode>(InSourcePin->GetOwningNode()))
+	// {
+	// 	if(InSourcePin->Direction == EGPD_Output)
+	// 	{
+	// 		return FPinConnectionResponse(CONNECT_RESPONSE_MAKE, LOCTEXT("MakeANewConstraint", "Create a new constraint"));
+	// 	}
+	// }
+
+	return FPinConnectionResponse(CONNECT_RESPONSE_DISALLOW, LOCTEXT("RigMapperDefinitionEditor_CannotCreateNode", "It is currently not possible to create new nodes from the graph"));
+}
+
+bool URigMapperDefinitionEditorGraphSchema::SupportsDropPinOnNode(UEdGraphNode* InTargetNode, const FEdGraphPinType& InSourcePinType, EEdGraphPinDirection InSourcePinDirection, FText& OutErrorMessage) const
+{
+	OutErrorMessage = LOCTEXT("RigMapperDefinitionEditor_CannotCreateConnection_DropPinOnNode", "It is currently not possible to create new connections or nodes from the graph");
+	return false;
+}
+
+#undef LOCTEXT_NAMESPACE
diff --git a/Engine/Plugins/Experimental/Animation/RigMapper/Source/RigMapperEditor/Private/RigMapperGraph/RigMapperDefinitionEditorGraphSchema.h b/Engine/Plugins/Experimental/Animation/RigMapper/Source/RigMapperEditor/Private/RigMapperGraph/RigMapperDefinitionEditorGraphSchema.h
new file mode 100644
--- /dev/null
+++ b/Engine/Plugins/Experimental/Animation/RigMapper/Source/RigMapperEditor/Private/RigMapperGraph/RigMapperDefinitionEditorGraphSchema.h
@@ -0,0 +1,52 @@
+// Copyright Epic Games, Inc. All Rights Reserved.
+
+#pragma once
+
+#include "CoreMinimal.h"
+#include "ConnectionDrawingPolicy.h"
+#include "EdGraph/EdGraphSchema.h"
+
+#include "RigMapperDefinitionEditorGraphSchema.generated.h"
+
+/**
+ * 
+ */
+UCLASS()
+class RIGMAPPEREDITOR_API URigMapperDefinitionEditorGraphSchema : public UEdGraphSchema
+{
+	GENERATED_BODY()
+
+public:
+	class ConnectionDrawingPolicy : public FConnectionDrawingPolicy
+	{
+	public:
+		ConnectionDrawingPolicy(int32 InBackLayerID, int32 InFrontLayerID, float InZoomFactor, const FSlateRect& InClippingRect, FSlateWindowElementList& InDrawElements)
+			: FConnectionDrawingPolicy(InBackLayerID, InFrontLayerID, InZoomFactor, InClippingRect, InDrawElements)
+		{
+			// We don't draw arrows here
+			ArrowImage = nullptr;
+			ArrowRadius = FVector2D(0.0f, 0.0f);
+		}
+
+		virtual FVector2D ComputeSplineTangent(const FVector2D& Start, const FVector2D& End) const override
+		{
+			const int32 Tension = FMath::Abs<int32>(Start.X - End.X);
+			return Tension * FVector2D(1.0f, 0);
+		}
+	};
+	
+	// UEdGraphSchema interface
+	virtual void GetContextMenuActions(class UToolMenu* Menu, class UGraphNodeContextMenuContext* Context) const override;
+	virtual FName GetParentContextMenuName() const override { return NAME_None; }
+	virtual FLinearColor GetPinTypeColor(const FEdGraphPinType& PinType) const override;
+	virtual void BreakPinLinks(UEdGraphPin& TargetPin, bool bSendsNodeNotification) const override;
+	virtual void BreakSinglePinLink(UEdGraphPin* SourcePin, UEdGraphPin* TargetPin) const override;
+	virtual FPinConnectionResponse MovePinLinks(UEdGraphPin& MoveFromPin, UEdGraphPin& MoveToPin, bool bIsIntermediateMove = false, bool bNotifyLinkedNodes = false) const override;
+	virtual FPinConnectionResponse CopyPinLinks(UEdGraphPin& CopyFromPin, UEdGraphPin& CopyToPin, bool bIsIntermediateCopy = false) const override;
+	virtual class FConnectionDrawingPolicy* CreateConnectionDrawingPolicy(int32 InBackLayerID, int32 InFrontLayerID, float InZoomFactor, const FSlateRect& InClippingRect, class FSlateWindowElementList& InDrawElements, class UEdGraph* InGraphObj) const override;
+	virtual bool ShouldAlwaysPurgeOnModification() const override { return false; }
+	virtual FPinConnectionResponse CanCreateNewNodes(UEdGraphPin* InSourcePin) const override;
+	virtual const FPinConnectionResponse CanCreateConnection(const UEdGraphPin* A, const UEdGraphPin* B) const override;
+	virtual bool SupportsDropPinOnNode(UEdGraphNode* InTargetNode, const FEdGraphPinType& InSourcePinType, EEdGraphPinDirection InSourcePinDirection, FText& OutErrorMessage) const override;
+	// End of UEdGraphSchema interface
+};
diff --git a/Engine/Plugins/Experimental/Animation/RigMapper/Source/RigMapperEditor/Private/RigMapperGraph/SRigMapperDefinitionGraphEditor.cpp b/Engine/Plugins/Experimental/Animation/RigMapper/Source/RigMapperEditor/Private/RigMapperGraph/SRigMapperDefinitionGraphEditor.cpp
new file mode 100644
--- /dev/null
+++ b/Engine/Plugins/Experimental/Animation/RigMapper/Source/RigMapperEditor/Private/RigMapperGraph/SRigMapperDefinitionGraphEditor.cpp
@@ -0,0 +1,227 @@
+// Copyright Epic Games, Inc. All Rights Reserved.
+
+#include "SRigMapperDefinitionGraphEditor.h"
+
+#include "RigMapperDefinitionEditorGraph.h"
+#include "RigMapperDefinitionEditorGraphNode.h"
+#include "RigMapperDefinitionEditorGraphSchema.h"
+
+#include "SGraphPanel.h"
+#include "SlateOptMacros.h"
+#include "Components/VerticalBox.h"
+
+BEGIN_SLATE_FUNCTION_BUILD_OPTIMIZATION
+
+#define LOCTEXT_NAMESPACE "SRigMapperDefinitionGraphEditor"
+
+SRigMapperDefinitionGraphEditor::~SRigMapperDefinitionGraphEditor()
+{
+	if (!GExitPurge)
+	{
+		if (ensure(GraphObj))
+		{
+			GraphObj->RemoveFromRoot();
+		}		
+	}
+}
+
+void SRigMapperDefinitionGraphEditor::Construct(const FArguments& InArgs, URigMapperDefinition* InDefinition)
+{
+	FGraphAppearanceInfo AppearanceInfo;
+	AppearanceInfo.CornerText = INVTEXT("Rig Mapper Definition");
+
+	SGraphEditor::FGraphEditorEvents GraphEvents;
+	// GraphEvents.OnCreateActionMenu = SGraphEditor::FOnCreateActionMenu::CreateSP(this, &SPhysicsAssetGraph::OnCreateGraphActionMenu);
+	GraphEvents.OnSelectionChanged = SGraphEditor::FOnSelectionChanged::CreateSP(this, &SRigMapperDefinitionGraphEditor::HandleSelectionChanged);
+	// GraphEvents.OnNodeDoubleClicked = FSingleNodeEvent::CreateSP(this, &SPhysicsAssetGraph::HandleNodeDoubleClicked);
+
+	GraphObj = NewObject<URigMapperDefinitionEditorGraph>();
+	GraphObj->AddToRoot();
+	GraphObj->Initialize(InDefinition);
+	GraphObj->RebuildGraph();
+
+	ChildSlot
+	[
+		SNew(SVerticalBox)
+		+SVerticalBox::Slot()
+		.FillHeight(1.f)
+		[
+			SAssignNew(GraphEditor, SGraphEditor)
+			.GraphToEdit(GraphObj)
+			.GraphEvents(GraphEvents)
+			.Appearance(AppearanceInfo)
+			.ShowGraphStateOverlay(false)
+			.IsEditable(false)
+			.AutoExpandActionMenu(false)
+			.DisplayAsReadOnly(false)
+		]
+	];
+
+}
+
+void SRigMapperDefinitionGraphEditor::Tick(const FGeometry& AllottedGeometry, const double InCurrentTime, const float InDeltaTime)
+{
+	if (GraphObj->NeedsRefreshLayout())
+	{		
+		GraphObj->LayoutNodes();
+
+		GraphObj->RequestRefreshLayout(false);
+
+		const FGraphPanelSelectionSet& Selection = GraphEditor->GetSelectedNodes();
+		
+		TArray<URigMapperDefinitionEditorGraphNode*> SelectedNodes;
+		SelectedNodes.Reserve(Selection.Num());
+		
+		for (UObject* Obj : Selection)
+		{
+			if (URigMapperDefinitionEditorGraphNode* Node = Cast<URigMapperDefinitionEditorGraphNode>(Obj))
+			{
+				SelectedNodes.Add(Node);
+			}
+		}
+		
+		ZoomToFitNodes(SelectedNodes);
+	}
+}
+
+void SRigMapperDefinitionGraphEditor::SelectNodes(const TArray<FString>& Inputs, const TArray<FString>& Features, const TArray<FString>& Outputs)
+{
+	if (!bSelectingNodes)
+	{
+		bSelectingNodes = true;
+		
+		GraphEditor->ClearSelectionSet();
+
+		const TArray<URigMapperDefinitionEditorGraphNode*>& SelectedNodes = GraphObj->GetNodesByName(Inputs, Features, Outputs); 
+		
+		for (URigMapperDefinitionEditorGraphNode* Node : SelectedNodes)
+		{
+			GraphEditor->SetNodeSelection(Node, true);
+		}
+		ZoomToFitNodes(SelectedNodes);
+		
+		bSelectingNodes = false;
+	}
+}
+
+void SRigMapperDefinitionGraphEditor::RebuildGraph()
+{
+	GraphObj->RebuildGraph();
+}
+
+void SRigMapperDefinitionGraphEditor::ZoomToFitNodes(const TArray<URigMapperDefinitionEditorGraphNode*>& SelectedNodes) const
+{
+	if (GraphObj->NeedsRefreshLayout())
+	{
+		return;
+	}
+	if (!bFocusLinkedNodes)
+	{
+		GraphEditor->ZoomToFit(true);
+	}
+	else
+	{
+		// todo: Allow option to show only selected graph/tree nodes
+		// todo  GraphEditor->GetGraphPanel()->GetNodeWidgetFromGuid(Node->NodeGuid)->SetVisibility()
+		
+		if (!SelectedNodes.IsEmpty() && !GraphEditor->GetGraphPanel()->HasDeferredZoomDestination())
+		{
+			FVector2D MinCorner(MAX_FLT, MAX_FLT);
+			FVector2D MaxCorner(-MAX_FLT, -MAX_FLT);
+
+			const FVector2D MaxLinkedNodeOffset(600, 400);
+			
+			for (UObject* Obj : SelectedNodes)
+			{
+				URigMapperDefinitionEditorGraphNode* Node = Cast<URigMapperDefinitionEditorGraphNode>(Obj);
+
+				FVector2D NodeTopLeft;
+				FVector2D NodeBottomRight;
+				Node->GetRect(NodeTopLeft, NodeBottomRight);
+				
+				TArray<URigMapperDefinitionEditorGraphNode*> LinkedNodes = { Node };
+				GetAllLinkedNodes(Node, LinkedNodes, true);
+				GetAllLinkedNodes(Node, LinkedNodes, false);
+
+				for (URigMapperDefinitionEditorGraphNode* LinkedNode : LinkedNodes)
+				{
+					FVector2D LinkedNodeTopLeft;
+					FVector2D LinkedNodeBottomRight;
+					LinkedNode->GetRect(LinkedNodeTopLeft, LinkedNodeBottomRight);
+					
+					MinCorner.X = FMath::Min(MinCorner.X, FMath::Max(LinkedNodeTopLeft.X, NodeTopLeft.X - MaxLinkedNodeOffset.X));
+					MinCorner.Y = FMath::Min(MinCorner.Y, FMath::Max(LinkedNodeTopLeft.Y, NodeTopLeft.Y - MaxLinkedNodeOffset.Y));
+					MaxCorner.X = FMath::Max(MaxCorner.X, FMath::Min(LinkedNodeBottomRight.X, NodeBottomRight.X + MaxLinkedNodeOffset.X));
+					MaxCorner.Y = FMath::Max(MaxCorner.Y, FMath::Min(LinkedNodeBottomRight.Y, NodeBottomRight.Y + MaxLinkedNodeOffset.Y));			
+				}
+			}
+			GraphEditor->GetGraphPanel()->JumpToRect(MinCorner, MaxCorner);
+		}
+	}
+}
+
+void SRigMapperDefinitionGraphEditor::GetAllLinkedNodes(const URigMapperDefinitionEditorGraphNode* BaseNode, TArray<URigMapperDefinitionEditorGraphNode*>& LinkedNodes, bool bDescend)
+{
+	
+	TArray<UEdGraphPin*> Pins;
+	if (bDescend)
+	{
+		Pins = BaseNode->GetInputPins();
+	}
+	else
+	{
+		Pins = { BaseNode->GetOutputPin() };
+	}
+	
+	for (const UEdGraphPin* PinA : Pins)
+	{
+		if (PinA)
+		{
+			for (const UEdGraphPin* PinB : PinA->LinkedTo)
+			{
+				if (PinB)
+				{
+					URigMapperDefinitionEditorGraphNode* LinkedNode = Cast<URigMapperDefinitionEditorGraphNode>(PinB->GetOwningNode());
+
+					if (LinkedNode && !LinkedNodes.Contains(LinkedNode))
+					{
+						LinkedNodes.Add(LinkedNode);
+						GetAllLinkedNodes(LinkedNode, LinkedNodes, bDescend);
+					}
+				}
+			}
+		}
+	}
+}
+
+void SRigMapperDefinitionGraphEditor::HandleSelectionChanged(const TSet<UObject*>& Nodes)
+{
+	if (!bSelectingNodes)
+	{
+		bSelectingNodes = true;
+		
+		if (OnSelectionChanged.IsBound())
+		{
+			OnSelectionChanged.Execute(Nodes);
+		}
+
+		TArray<URigMapperDefinitionEditorGraphNode*> SelectedNodes;
+
+		for (UObject* Obj : Nodes)
+		{
+			if (URigMapperDefinitionEditorGraphNode* Node = Cast<URigMapperDefinitionEditorGraphNode>(Obj))
+			{
+				SelectedNodes.Add(Node);
+			}
+		}
+		
+		ZoomToFitNodes(SelectedNodes);
+		
+		bSelectingNodes = false;
+	}
+}
+
+
+#undef LOCTEXT_NAMESPACE
+
+END_SLATE_FUNCTION_BUILD_OPTIMIZATION
diff --git a/Engine/Plugins/Experimental/Animation/RigMapper/Source/RigMapperEditor/Private/RigMapperGraph/SRigMapperDefinitionGraphEditor.h b/Engine/Plugins/Experimental/Animation/RigMapper/Source/RigMapperEditor/Private/RigMapperGraph/SRigMapperDefinitionGraphEditor.h
new file mode 100644
--- /dev/null
+++ b/Engine/Plugins/Experimental/Animation/RigMapper/Source/RigMapperEditor/Private/RigMapperGraph/SRigMapperDefinitionGraphEditor.h
@@ -0,0 +1,52 @@
+// Copyright Epic Games, Inc. All Rights Reserved.
+
+#pragma once
+
+#include "CoreMinimal.h"
+#include "GraphEditor.h"
+#include "Widgets/SCompoundWidget.h"
+
+class URigMapperDefinitionEditorGraphNode;
+class URigMapperDefinition;
+class URigMapperDefinitionEditorGraph;
+
+/**
+ * 
+ */
+class RIGMAPPEREDITOR_API SRigMapperDefinitionGraphEditor : public SCompoundWidget
+{
+public:
+	SLATE_BEGIN_ARGS(SRigMapperDefinitionGraphEditor)
+		{
+		}
+
+	SLATE_END_ARGS()
+
+	virtual ~SRigMapperDefinitionGraphEditor() override;
+
+	/** Constructs this widget with InArgs */
+	void Construct(const FArguments& InArgs, URigMapperDefinition* InDefinition);
+
+	virtual void Tick(const FGeometry& AllottedGeometry, const double InCurrentTime, const float InDeltaTime) override;
+
+	void SelectNodes(const TArray<FString>& Inputs, const TArray<FString>& Features, const TArray<FString>& Outputs);
+	void RebuildGraph();
+
+private:
+	void ZoomToFitNodes(const TArray<URigMapperDefinitionEditorGraphNode*>& SelectedNodes) const;
+	static void GetAllLinkedNodes(const URigMapperDefinitionEditorGraphNode* BaseNode, TArray<URigMapperDefinitionEditorGraphNode*>& LinkedNodes, bool bDescend);
+	void HandleSelectionChanged(const TSet<UObject*>& Nodes);
+
+public:
+	SGraphEditor::FOnSelectionChanged OnSelectionChanged;
+	
+private:
+	TSharedPtr<SGraphEditor> GraphEditor;
+	TObjectPtr<URigMapperDefinitionEditorGraph> GraphObj;
+
+	bool bSelectingNodes = false;
+
+	bool bFocusLinkedNodes = true;
+
+	bool bDeferZoom = false;
+};
diff --git a/Engine/Plugins/Experimental/Animation/RigMapper/Source/RigMapperEditor/Private/RigMapperGraph/SRigMapperDefinitionGraphEditorNode.cpp b/Engine/Plugins/Experimental/Animation/RigMapper/Source/RigMapperEditor/Private/RigMapperGraph/SRigMapperDefinitionGraphEditorNode.cpp
new file mode 100644
--- /dev/null
+++ b/Engine/Plugins/Experimental/Animation/RigMapper/Source/RigMapperEditor/Private/RigMapperGraph/SRigMapperDefinitionGraphEditorNode.cpp
@@ -0,0 +1,128 @@
+// Copyright Epic Games, Inc. All Rights Reserved.
+
+#include "SRigMapperDefinitionGraphEditorNode.h"
+
+#include "RigMapperDefinitionEditorGraphNode.h"
+#include "SlateOptMacros.h"
+#include "Components/HorizontalBox.h"
+
+BEGIN_SLATE_FUNCTION_BUILD_OPTIMIZATION
+
+void SRigMapperDefinitionGraphEditorNode::Construct(const FArguments& InArgs, URigMapperDefinitionEditorGraphNode* InNode)
+{
+	GraphNode = InNode;
+	SetCursor(EMouseCursor::CardinalCross);
+
+	// if (!ContentWidget.IsValid())
+	// {
+	// 	ContentWidget = SNullWidget::NullWidget;
+	// }
+	
+	UpdateGraphNode();
+}
+
+FSlateColor SRigMapperDefinitionGraphEditorNode::GetNodeColor() const
+{
+	if (GraphNode)
+	{
+		return FSlateColor(GraphNode->GetNodeTitleColor());
+	}
+	return FSlateColor(FLinearColor::Red);
+}
+
+FText SRigMapperDefinitionGraphEditorNode::GetNodeTitle() const
+{
+	if (GraphNode)
+	{
+		return GraphNode->GetNodeTitle(ENodeTitleType::FullTitle);
+	}
+	return FText::FromString("Invalid");
+}
+
+FText SRigMapperDefinitionGraphEditorNode::GetNodeSubtitle() const
+{
+	if (const URigMapperDefinitionEditorGraphNode* Node = Cast<URigMapperDefinitionEditorGraphNode>(GraphNode))
+	{
+		return Node->GetSubtitle();
+	}
+	return FText::GetEmpty();
+}
+
+
+void SRigMapperDefinitionGraphEditorNode::UpdateGraphNode()
+{
+	InputPins.Empty();
+	OutputPins.Empty();
+	
+	RightNodeBox.Reset();
+	LeftNodeBox.Reset();
+
+	this->ContentScale.Bind( this, &SGraphNode::GetContentScale );
+
+	this->GetOrAddSlot( ENodeZone::Center )
+		.HAlign(HAlign_Center)
+		.VAlign(VAlign_Center)
+		[
+			SNew(SBorder)
+			.BorderImage( FAppStyle::GetBrush("PhysicsAssetEditor.Graph.NodeBody") )
+			.BorderBackgroundColor(this, &SRigMapperDefinitionGraphEditorNode::GetNodeColor)
+			.Padding(0)
+			[
+				SNew(SHorizontalBox)
+
+				+SHorizontalBox::Slot()
+				.AutoWidth()
+				.HAlign(HAlign_Center)
+				.VAlign(VAlign_Center)
+				[
+					SAssignNew(LeftNodeBox, SVerticalBox)
+				]
+
+				+SHorizontalBox::Slot()
+				.AutoWidth()
+				.HAlign(HAlign_Fill)
+				.VAlign(VAlign_Center)
+				.Padding(4.0f)
+				[
+					SNew(SVerticalBox)
+					+SVerticalBox::Slot()
+					.AutoHeight()
+					.HAlign(HAlign_Center)
+					.Padding(8, 0)
+					[
+						SNew(STextBlock)
+						.TextStyle( FAppStyle::Get(), "Graph.Node.NodeTitle" )
+						.ColorAndOpacity(FSlateColor(FLinearColor::Black))
+						.Text(this, &SRigMapperDefinitionGraphEditorNode::GetNodeTitle)
+					]
+					+SVerticalBox::Slot()
+					.AutoHeight()
+					.HAlign(HAlign_Center)
+					[
+						SNew(STextBlock)
+						.TextStyle(FAppStyle::Get(), "NormalText")
+						.ColorAndOpacity(FSlateColor(FLinearColor::Black))
+						.Text(this, &SRigMapperDefinitionGraphEditorNode::GetNodeSubtitle)
+						.Visibility(this, &SRigMapperDefinitionGraphEditorNode::GetShowNodeSubtitle)
+					]
+				]
+
+				+SHorizontalBox::Slot()
+				.AutoWidth()
+				.HAlign(HAlign_Center)
+				.VAlign(VAlign_Center)
+				[
+					SAssignNew(RightNodeBox, SVerticalBox)
+				]
+			]
+		];
+
+	CreatePinWidgets();
+}
+
+const FSlateBrush* SRigMapperDefinitionGraphEditorNode::GetShadowBrush(bool bSelected) const
+{
+	return bSelected ? FAppStyle::GetBrush(TEXT("PhysicsAssetEditor.Graph.Node.ShadowSelected")) : FAppStyle::GetBrush(TEXT("PhysicsAssetEditor.Graph.Node.Shadow"));
+}
+
+END_SLATE_FUNCTION_BUILD_OPTIMIZATION
diff --git a/Engine/Plugins/Experimental/Animation/RigMapper/Source/RigMapperEditor/Private/RigMapperGraph/SRigMapperDefinitionGraphEditorNode.h b/Engine/Plugins/Experimental/Animation/RigMapper/Source/RigMapperEditor/Private/RigMapperGraph/SRigMapperDefinitionGraphEditorNode.h
new file mode 100644
--- /dev/null
+++ b/Engine/Plugins/Experimental/Animation/RigMapper/Source/RigMapperEditor/Private/RigMapperGraph/SRigMapperDefinitionGraphEditorNode.h
@@ -0,0 +1,34 @@
+// Copyright Epic Games, Inc. All Rights Reserved.
+
+#pragma once
+
+#include "CoreMinimal.h"
+#include "SGraphNode.h"
+
+class URigMapperDefinitionEditorGraphNode;
+
+/**
+ * 
+ */
+class RIGMAPPEREDITOR_API SRigMapperDefinitionGraphEditorNode : public SGraphNode
+{
+public:
+	SLATE_BEGIN_ARGS(SRigMapperDefinitionGraphEditorNode)
+		{
+		}
+
+	SLATE_END_ARGS()
+
+	/** Constructs this widget with InArgs */
+	void Construct(const FArguments& InArgs, URigMapperDefinitionEditorGraphNode* InNode);
+
+protected:
+	FSlateColor GetNodeColor() const;
+	FText GetNodeTitle() const;
+	FText GetNodeSubtitle() const;
+	EVisibility GetShowNodeSubtitle() const { return GetNodeSubtitle().IsEmpty() ? EVisibility::Collapsed : EVisibility::Visible; };
+	
+	// SGraphNode interface
+	virtual void UpdateGraphNode() override;
+	virtual const FSlateBrush* GetShadowBrush(bool bSelected) const override;
+};
diff --git a/Engine/Plugins/Experimental/Animation/RigMapper/Source/RigMapperEditor/Private/RigMapperGraph/SRigMapperDefinitionStructureView.cpp b/Engine/Plugins/Experimental/Animation/RigMapper/Source/RigMapperEditor/Private/RigMapperGraph/SRigMapperDefinitionStructureView.cpp
new file mode 100644
--- /dev/null
+++ b/Engine/Plugins/Experimental/Animation/RigMapper/Source/RigMapperEditor/Private/RigMapperGraph/SRigMapperDefinitionStructureView.cpp
@@ -0,0 +1,336 @@
+// Copyright Epic Games, Inc. All Rights Reserved.
+
+#include "SRigMapperDefinitionStructureView.h"
+
+#include "RigMapperDefinition.h"
+#include "RigMapperDefinitionEditorGraphNode.h"
+#include "SlateOptMacros.h"
+#include "Widgets/Input/SSearchBox.h"
+
+const int32 SRigMapperDefinitionStructureView::NumNodeTypes = 3;
+const int32 SRigMapperDefinitionStructureView::NumFeatureTypes = 3;
+const FString SRigMapperDefinitionStructureView::InputsNodeName = TEXT("Inputs");
+const FString SRigMapperDefinitionStructureView::FeaturesNodeName = TEXT("Features");
+const FString SRigMapperDefinitionStructureView::MultiplyNodeName = TEXT("Multiply");
+const FString SRigMapperDefinitionStructureView::WsNodeName = TEXT("Weighted Sums");
+const FString SRigMapperDefinitionStructureView::SdkNodeName = TEXT("SDKs");
+const FString SRigMapperDefinitionStructureView::OutputNodeName = TEXT("Outputs");
+
+
+BEGIN_SLATE_FUNCTION_BUILD_OPTIMIZATION
+
+void SRigMapperDefinitionStructureView::Construct(const FArguments& InArgs, URigMapperDefinition* InDefinition)
+{
+	Definition = InDefinition;
+	GenerateParentNodes();
+	GenerateChildrenNodes();
+
+	TreeView = SNew(STreeView<TSharedPtr<FString>>)
+		.SelectionMode(ESelectionMode::Multi)
+		.HighlightParentNodesForSelection(true)
+		.OnGenerateRow(this, &SRigMapperDefinitionStructureView::OnGenerateTreeRow)
+		.OnGetChildren(this, &SRigMapperDefinitionStructureView::OnGetTreeNodeChildren)
+		.OnSelectionChanged(this, &SRigMapperDefinitionStructureView::HandleTreeNodesSelectionChanged)
+		.TreeItemsSource(&FilteredRootNodes);
+
+	SearchBoxFilter = MakeShared<TTextFilter<TSharedPtr<FString>>>(TTextFilter<TSharedPtr<FString>>::FItemToStringArray::CreateSP(this, &SRigMapperDefinitionStructureView::TransformElementToString));
+	
+	for (const TPair<TSharedPtr<FString>, TArray<TSharedPtr<FString>>>& NodeAndChildren : ParentsAndChildrenNodes)
+	{
+		TreeView->SetItemExpansion(NodeAndChildren.Key, true);	
+	}
+	
+	ChildSlot
+	[
+		SNew(SVerticalBox)
+		+ SVerticalBox::Slot()
+		.Padding(7.0f, 6.0f)
+		.AutoHeight()
+		[
+			SAssignNew(SearchBox, SSearchBox)
+			.OnTextChanged(this, &SRigMapperDefinitionStructureView::OnFilterTextChanged)
+		]
+
+		+ SVerticalBox::Slot()
+		[
+			TreeView.ToSharedRef()
+		]
+	]; 
+}
+
+END_SLATE_FUNCTION_BUILD_OPTIMIZATION
+
+
+TSharedPtr<FString> SRigMapperDefinitionStructureView::GetParentAndChildrenNodes(ERigMapperNodeType NodeType, TArray<TSharedPtr<FString>>& OutChildren)
+{
+	if (const TSharedPtr<FString>* ParentNode = ParentNodesMapping.Find(NodeType))
+	{
+		if (const TArray<TSharedPtr<FString>>* ChildrenNodes = ParentsAndChildrenNodes.Find(*ParentNode))
+		{
+			OutChildren = *ChildrenNodes;
+			return *ParentNode;
+		}
+	}
+	return nullptr;
+}
+
+TArray<TSharedPtr<FString>>* SRigMapperDefinitionStructureView::GetChildrenNodes(ERigMapperNodeType NodeType)
+{
+	if (const TSharedPtr<FString>* ParentNode = ParentNodesMapping.Find(NodeType))
+	{
+		return ParentsAndChildrenNodes.Find(*ParentNode);
+	}
+	return nullptr;
+}
+
+const TArray<TSharedPtr<FString>>* SRigMapperDefinitionStructureView::GetChildrenNodes(ERigMapperNodeType NodeType) const
+{
+	if (const TSharedPtr<FString>* ParentNode = ParentNodesMapping.Find(NodeType))
+	{
+		return ParentsAndChildrenNodes.Find(*ParentNode);
+	}
+	return nullptr;
+}
+
+bool SRigMapperDefinitionStructureView::SelectNode(const FString& NodeName, ERigMapperNodeType NodeType, bool bSelected)
+{
+	if (const TArray<TSharedPtr<FString>>* ChildrenNodes = GetChildrenNodes(NodeType))
+	{
+		if (const TSharedPtr<FString>* TreeNode = ChildrenNodes->FindByPredicate(
+			[&NodeName](const TSharedPtr<FString>& Item) { return Item && *Item.Get() == NodeName; }))
+		{
+			TreeView->SetItemSelection(*TreeNode, bSelected);
+			TreeView->RequestNavigateToItem(*TreeNode);
+			return true;
+		}
+	}
+	return false;
+}
+
+void SRigMapperDefinitionStructureView::RebuildTree()
+{
+	TreeView->ClearSelection();
+	GenerateChildrenNodes();
+	
+	FilteredRootNodes.Reset();
+	FilterNodes(RootNodes, FilteredRootNodes);
+
+	TreeView->RequestTreeRefresh();
+}
+
+void SRigMapperDefinitionStructureView::ClearSelection() const
+{
+	TreeView->ClearSelection();
+}
+
+bool SRigMapperDefinitionStructureView::IsNodeOrChildSelected(ERigMapperNodeType NodeType, int32 ArrayIndex) const
+{
+	if (NodeType == ERigMapperNodeType::Invalid)
+	{
+		return IsNodeOrChildSelected(ERigMapperNodeType::Multiply, ArrayIndex) ||
+				IsNodeOrChildSelected(ERigMapperNodeType::WeightedSum, ArrayIndex) ||
+				IsNodeOrChildSelected(ERigMapperNodeType::SDK, ArrayIndex);
+	}
+	
+	const TArray<TSharedPtr<FString>>& Selection = TreeView->GetSelectedItems();
+	
+	if (const TArray<TSharedPtr<FString>>* ChildrenNodes = GetChildrenNodes(NodeType))
+	{
+		if (ArrayIndex == INDEX_NONE)
+		{
+			for (const TSharedPtr<FString>& Node : Selection)
+			{
+				if (ChildrenNodes->Contains(Node))
+				{
+					return true;
+				}
+			}
+		}
+		return ChildrenNodes->IsValidIndex(ArrayIndex) && Selection.Contains((*ChildrenNodes)[ArrayIndex]);	
+	}
+	return false;
+}
+
+bool SRigMapperDefinitionStructureView::IsSelectionEmpty() const
+{
+	return TreeView->GetNumItemsSelected() == 0;
+}
+
+void SRigMapperDefinitionStructureView::GenerateParentNodes()
+{
+	RootNodes.Reset(NumFeatureTypes);
+	ParentNodesMapping.Empty(NumFeatureTypes + NumNodeTypes);
+	ParentsAndChildrenNodes.Empty(NumFeatureTypes + NumNodeTypes);
+
+	ParentsAndChildrenNodes.Add(ParentNodesMapping.Add(
+		ERigMapperNodeType::Input, RootNodes.Add_GetRef(MakeShared<FString>(InputsNodeName))));
+
+	TArray<TSharedPtr<FString>>& FeatureEntries = ParentsAndChildrenNodes.Add(ParentNodesMapping.Add(
+		ERigMapperNodeType::Invalid, RootNodes.Add_GetRef(MakeShared<FString>(FeaturesNodeName))));
+	FeatureEntries.Reserve(NumNodeTypes);
+
+	ParentsAndChildrenNodes.Add(ParentNodesMapping.Add(
+		ERigMapperNodeType::Multiply, FeatureEntries.Add_GetRef(MakeShared<FString>(MultiplyNodeName))));
+
+	ParentsAndChildrenNodes.Add(ParentNodesMapping.Add(
+		ERigMapperNodeType::WeightedSum, FeatureEntries.Add_GetRef(MakeShared<FString>(WsNodeName))));
+	
+	ParentsAndChildrenNodes.Add(ParentNodesMapping.Add(
+		ERigMapperNodeType::SDK, FeatureEntries.Add_GetRef(MakeShared<FString>(SdkNodeName))));
+	
+	ParentsAndChildrenNodes.Add(ParentNodesMapping.Add(
+		ERigMapperNodeType::Output, RootNodes.Add_GetRef(MakeShared<FString>(OutputNodeName))));
+
+	FilteredRootNodes = RootNodes;
+}
+
+void SRigMapperDefinitionStructureView::GenerateChildrenNodes()
+{
+	if (Definition && !ParentsAndChildrenNodes.IsEmpty() && !ParentNodesMapping.IsEmpty())
+	{
+		TArray<TSharedPtr<FString>>* InputEntries = GetChildrenNodes(ERigMapperNodeType::Input);
+		InputEntries->Reset(Definition->Inputs.Num());
+		for (const FString& Input : Definition->Inputs)
+		{
+			InputEntries->Add(MakeShared<FString>(Input));
+		}
+
+		TArray<TSharedPtr<FString>>* MultiplyEntries = GetChildrenNodes(ERigMapperNodeType::Multiply);
+		MultiplyEntries->Reset(Definition->Features.Multiply.Num());
+		for (const FRigMapperMultiplyFeature& Feature : Definition->Features.Multiply)
+		{
+			MultiplyEntries->Add(MakeShared<FString>(Feature.Name));
+		}
+
+		TArray<TSharedPtr<FString>>* WsEntries = GetChildrenNodes(ERigMapperNodeType::WeightedSum);
+		WsEntries->Reset(Definition->Features.WeightedSums.Num());
+		for (const FRigMapperFeature& Feature : Definition->Features.WeightedSums)
+		{
+			WsEntries->Add(MakeShared<FString>(Feature.Name));
+		}
+
+		TArray<TSharedPtr<FString>>* SdkEntries = GetChildrenNodes(ERigMapperNodeType::SDK);
+		SdkEntries->Reset(Definition->Features.SDKs.Num());
+		for (const FRigMapperSdkFeature& Feature : Definition->Features.SDKs)
+		{
+			SdkEntries->Add(MakeShared<FString>(Feature.Name));
+		}
+
+		TArray<TSharedPtr<FString>>* OutputEntries = GetChildrenNodes(ERigMapperNodeType::Output);
+		OutputEntries->Reset(Definition->Outputs.Num());
+		for (const TPair<FString, FString>& Output : Definition->Outputs)
+		{
+			OutputEntries->Add(MakeShared<FString>(Output.Key));
+		}
+	}
+}
+
+void SRigMapperDefinitionStructureView::HandleTreeNodesSelectionChanged(TSharedPtr<FString> Node, ESelectInfo::Type SelectInfo)
+{
+	TArray<FString> SelectedInputs;
+	TArray<FString> SelectedFeatures;
+	TArray<FString> SelectedOutputs;
+
+	if (SelectInfo != ESelectInfo::Direct)
+	{
+		const TArray<TSharedPtr<FString>>& SelectedItems = TreeView->GetSelectedItems();
+
+		auto GetSelectedNodeNames = [&SelectedItems](const TArray<TSharedPtr<FString>>* PtrArray)
+		{
+			TArray<FString> SelectedNames;
+
+			for (const TSharedPtr<FString>& Item : SelectedItems)
+			{
+				if (PtrArray->Contains(Item))
+				{
+					SelectedNames.Add(*Item.Get());
+				}
+			}
+			return SelectedNames;
+		};
+
+		SelectedInputs = GetSelectedNodeNames(GetChildrenNodes(ERigMapperNodeType::Input));
+
+		SelectedFeatures = GetSelectedNodeNames(GetChildrenNodes(ERigMapperNodeType::Multiply));
+		SelectedFeatures.Append(GetSelectedNodeNames(GetChildrenNodes(ERigMapperNodeType::WeightedSum)));
+		SelectedFeatures.Append(GetSelectedNodeNames(GetChildrenNodes(ERigMapperNodeType::SDK)));
+	
+		SelectedOutputs = GetSelectedNodeNames(GetChildrenNodes(ERigMapperNodeType::Output));
+	}
+
+	if (OnSelectionChanged.IsBound())
+	{
+		OnSelectionChanged.Execute(SelectInfo, SelectedInputs, SelectedFeatures, SelectedOutputs);
+	}
+}
+
+
+void SRigMapperDefinitionStructureView::TransformElementToString(TSharedPtr<FString> String, TArray<FString>& Strings)
+{
+	if (String.IsValid())
+	{
+		Strings = { *String.Get() };
+	}
+}
+
+void SRigMapperDefinitionStructureView::OnFilterTextChanged(const FText& Text)
+{
+	SearchBoxFilter->SetRawFilterText(Text);
+
+	FilteredRootNodes.Reset();
+	FilterNodes(RootNodes, FilteredRootNodes);
+	TreeView->RequestTreeRefresh();
+}
+
+void SRigMapperDefinitionStructureView::FilterNodes(const TArray<TSharedPtr<FString>>& ParentNodes, TArray<TSharedPtr<FString>>& FilteredNodes)
+{
+	for (TSharedPtr<FString> ParentNode : ParentNodes)
+	{
+		bool ChildPassedFilter = false;
+		
+		if (ParentNode.IsValid() && ParentsAndChildrenNodes.Contains(ParentNode))
+		{
+			const TArray<TSharedPtr<FString>>& Children = ParentsAndChildrenNodes[ParentNode];
+
+			if (!SearchBoxFilter->GetRawFilterText().IsEmpty())
+			{
+				ChildPassedFilter = Children.ContainsByPredicate([this](const TSharedPtr<FString>& InItem)
+				{
+					return SearchBoxFilter->PassesFilter(InItem);
+				});
+			}
+
+			if (ChildPassedFilter)
+			{
+				FilteredNodes.AddUnique(ParentNode);
+				TreeView->SetItemExpansion(ParentNode, true);
+			}
+		}
+		if (!ChildPassedFilter && SearchBoxFilter->PassesFilter(ParentNode))
+		{
+			FilteredNodes.Add(ParentNode);
+		}
+	}
+}
+
+TSharedRef<ITableRow> SRigMapperDefinitionStructureView::OnGenerateTreeRow(TSharedPtr<FString> NodeName, const TSharedRef<STableViewBase>& TableViewBase)
+{
+	return SNew(STableRow<TSharedPtr<FString>>, TableViewBase)
+		.Style(&FAppStyle::Get().GetWidgetStyle<FTableRowStyle>("SceneOutliner.TableViewRow"))
+		.Padding(4)
+		[
+			SNew(STextBlock)
+			.Text(FText::FromString(*NodeName.Get()))
+		];
+}
+
+void SRigMapperDefinitionStructureView::OnGetTreeNodeChildren(TSharedPtr<FString> NodeName, TArray<TSharedPtr<FString>>& Children)
+{
+	if (NodeName.IsValid() && ParentsAndChildrenNodes.Contains(NodeName))
+	{
+		const TArray<TSharedPtr<FString>>& UnfilteredChildren = ParentsAndChildrenNodes[NodeName];
+
+		Children.Reset(UnfilteredChildren.Num());
+		FilterNodes(UnfilteredChildren, Children);
+	}
+}
diff --git a/Engine/Plugins/Experimental/Animation/RigMapper/Source/RigMapperEditor/Private/RigMapperGraph/SRigMapperDefinitionStructureView.h b/Engine/Plugins/Experimental/Animation/RigMapper/Source/RigMapperEditor/Private/RigMapperGraph/SRigMapperDefinitionStructureView.h
new file mode 100644
--- /dev/null
+++ b/Engine/Plugins/Experimental/Animation/RigMapper/Source/RigMapperEditor/Private/RigMapperGraph/SRigMapperDefinitionStructureView.h
@@ -0,0 +1,78 @@
+// Copyright Epic Games, Inc. All Rights Reserved.
+
+#pragma once
+
+#include "CoreMinimal.h"
+#include "Misc/TextFilter.h"
+#include "Widgets/SCompoundWidget.h"
+#include "Widgets/Views/STreeView.h"
+
+enum class ERigMapperNodeType : uint8;
+class SSearchBox;
+class URigMapperDefinition;
+
+/**
+ * 
+ */
+class RIGMAPPEREDITOR_API SRigMapperDefinitionStructureView : public SCompoundWidget
+{
+public:
+	SLATE_BEGIN_ARGS(SRigMapperDefinitionStructureView)
+		{
+		}
+
+	SLATE_END_ARGS()
+
+	DECLARE_DELEGATE_FourParams(FOnSelectionChanged, ESelectInfo::Type, TArray<FString>, TArray<FString>, TArray<FString>);
+	
+	/** Constructs this widget with InArgs */
+	void Construct(const FArguments& InArgs, URigMapperDefinition* InDefinition);
+
+	bool SelectNode(const FString& NodeName, ERigMapperNodeType NodeType, bool bSelected);
+	void RebuildTree();
+	void ClearSelection() const;
+	bool IsNodeOrChildSelected(ERigMapperNodeType NodeType, int32 ArrayIndex) const;
+	bool IsSelectionEmpty() const;
+	
+public:
+	FOnSelectionChanged OnSelectionChanged;
+
+private:
+	void GenerateParentNodes();
+	void GenerateChildrenNodes();
+	void HandleTreeNodesSelectionChanged(TSharedPtr<FString> Node, ESelectInfo::Type SelectInfo);
+	void TransformElementToString(TSharedPtr<FString> String, TArray<FString>& Strings);
+	void OnFilterTextChanged(const FText& Text);
+	void FilterNodes(const TArray<TSharedPtr<FString>>& ParentNodes, TArray<TSharedPtr<FString>>& FilteredNodes);
+
+	TSharedRef<ITableRow> OnGenerateTreeRow(TSharedPtr<FString> NodeName, const TSharedRef<STableViewBase>& TableViewBase);
+	void OnGetTreeNodeChildren(TSharedPtr<FString> NodeName, TArray<TSharedPtr<FString>>& Children);
+
+	TSharedPtr<FString> GetParentAndChildrenNodes(ERigMapperNodeType NodeType, TArray<TSharedPtr<FString>>& OutChildren);
+	TArray<TSharedPtr<FString>>* GetChildrenNodes(ERigMapperNodeType NodeType);
+	const TArray<TSharedPtr<FString>>* GetChildrenNodes(ERigMapperNodeType NodeType) const;
+	
+private:
+	static const int32 NumNodeTypes;
+	static const int32 NumFeatureTypes;
+	
+	static const FString InputsNodeName;
+	static const FString FeaturesNodeName;
+	static const FString MultiplyNodeName;
+	static const FString WsNodeName;
+	static const FString SdkNodeName;
+	static const FString OutputNodeName;
+	
+	TObjectPtr<URigMapperDefinition> Definition;
+
+	TSharedPtr<STreeView<TSharedPtr<FString>>> TreeView;
+	
+	TSharedPtr<SSearchBox> SearchBox;
+	TSharedPtr<TTextFilter<TSharedPtr<FString>>> SearchBoxFilter;
+
+	TArray<TSharedPtr<FString>> RootNodes;
+	TArray<TSharedPtr<FString>> FilteredRootNodes;
+	
+	TMap<TSharedPtr<FString>, TArray<TSharedPtr<FString>>> ParentsAndChildrenNodes;
+	TMap<ERigMapperNodeType, TSharedPtr<FString>> ParentNodesMapping;
+};
diff --git a/Engine/Plugins/Experimental/Animation/RigMapper/Source/RigMapperEditor/Private/RigMapperLinkedDefinitionsAssetTypeActions.cpp b/Engine/Plugins/Experimental/Animation/RigMapper/Source/RigMapperEditor/Private/RigMapperLinkedDefinitionsAssetTypeActions.cpp
new file mode 100644
--- /dev/null
+++ b/Engine/Plugins/Experimental/Animation/RigMapper/Source/RigMapperEditor/Private/RigMapperLinkedDefinitionsAssetTypeActions.cpp
@@ -0,0 +1,25 @@
+// Copyright Epic Games, Inc. All Rights Reserved.
+
+#include "RigMapperLinkedDefinitionsAssetTypeActions.h"
+
+#include "RigMapperDefinition.h"
+
+UClass* FRigMapperLinkedDefinitionsAssetTypeActions::GetSupportedClass() const
+{
+	return URigMapperLinkedDefinitions::StaticClass();
+}
+
+FText FRigMapperLinkedDefinitionsAssetTypeActions::GetName() const
+{
+	return INVTEXT("Rig Mapper Linked Definitions");
+}
+
+FColor FRigMapperLinkedDefinitionsAssetTypeActions::GetTypeColor() const
+{
+	return FColor::Yellow;
+}
+
+uint32 FRigMapperLinkedDefinitionsAssetTypeActions::GetCategories()
+{
+	return EAssetTypeCategories::Animation;
+}
diff --git a/Engine/Plugins/Experimental/Animation/RigMapper/Source/RigMapperEditor/Private/RigMapperLinkedDefinitionsAssetTypeActions.h b/Engine/Plugins/Experimental/Animation/RigMapper/Source/RigMapperEditor/Private/RigMapperLinkedDefinitionsAssetTypeActions.h
new file mode 100644
--- /dev/null
+++ b/Engine/Plugins/Experimental/Animation/RigMapper/Source/RigMapperEditor/Private/RigMapperLinkedDefinitionsAssetTypeActions.h
@@ -0,0 +1,18 @@
+// Copyright Epic Games, Inc. All Rights Reserved.
+
+#pragma once
+ 
+#include "CoreMinimal.h"
+#include "AssetTypeActions_Base.h"
+
+/**
+ * The asset actions for the URigMapperLinkedDefinitions data asset class and link to its asset editor toolkit
+ */
+class RIGMAPPEREDITOR_API FRigMapperLinkedDefinitionsAssetTypeActions : public FAssetTypeActions_Base
+{
+public:
+	virtual UClass* GetSupportedClass() const override;
+	virtual FText GetName() const override;
+	virtual FColor GetTypeColor() const override;
+	virtual uint32 GetCategories() override;
+};
diff --git a/Engine/Plugins/Experimental/Animation/RigMapper/Source/RigMapperEditor/Public/AnimGraphNode_RigMapper.h b/Engine/Plugins/Experimental/Animation/RigMapper/Source/RigMapperEditor/Public/AnimGraphNode_RigMapper.h
new file mode 100644
--- /dev/null
+++ b/Engine/Plugins/Experimental/Animation/RigMapper/Source/RigMapperEditor/Public/AnimGraphNode_RigMapper.h
@@ -0,0 +1,38 @@
+// Copyright Epic Games, Inc. All Rights Reserved.
+
+#pragma once
+
+#include "AnimGraphNode_Base.h"
+#include "CoreMinimal.h"
+#include "AnimNode_RigMapper.h"
+
+#include "AnimGraphNode_RigMapper.generated.h"
+
+/**
+ * 
+ */
+UCLASS()
+class RIGMAPPEREDITOR_API UAnimGraphNode_RigMapper : public UAnimGraphNode_Base
+{
+	GENERATED_UCLASS_BODY()
+
+public:
+	UPROPERTY(EditAnywhere, BlueprintReadWrite, Category = Settings)
+	FAnimNode_RigMapper Node;
+
+public:
+	// UEdGraphNode interface
+	virtual FText GetNodeTitle(ENodeTitleType::Type TitleType) const override;
+	virtual FText GetTooltipText() const override;
+	virtual FText GetMenuCategory() const override;
+	virtual FLinearColor GetNodeBodyTintColor() const override;
+	// End of UEdGraphNode interface
+
+	// UAnimGraphNode_Base interface
+	virtual void ValidateAnimNodeDuringCompilation(USkeleton* ForSkeleton, FCompilerResultsLog& MessageLog) override;
+	virtual FEditorModeID GetEditorMode() const override;
+	virtual EAnimAssetHandlerType SupportsAssetClass(const UClass* AssetClass) const override;
+	virtual void CopyNodeDataToPreviewNode(FAnimNode_Base* InPreviewNode) override;
+	virtual void PostEditChangeProperty(FPropertyChangedEvent& PropertyChangedEvent) override;
+	// End of UAnimGraphNode_Base interface
+};
diff --git a/Engine/Plugins/Experimental/Animation/RigMapper/Source/RigMapperEditor/Public/RigMapperEditorSubsystem.h b/Engine/Plugins/Experimental/Animation/RigMapper/Source/RigMapperEditor/Public/RigMapperEditorSubsystem.h
new file mode 100644
--- /dev/null
+++ b/Engine/Plugins/Experimental/Animation/RigMapper/Source/RigMapperEditor/Public/RigMapperEditorSubsystem.h
@@ -0,0 +1,182 @@
+// Copyright Epic Games, Inc. All Rights Reserved.
+
+#pragma once
+
+#include "CoreMinimal.h"
+#include "Misc/FrameRate.h"
+#include "EditorSubsystem.h"
+#include "Kismet/BlueprintFunctionLibrary.h"
+
+#include "RigMapper.h"
+#include "RigMapperDefinition.h"
+
+#include "RigMapperEditorSubsystem.generated.h"
+
+class USkeletalMesh;
+class UAnimSequence;
+class ULevelSequence;
+class UControlRig;
+class UMovieSceneControlRigParameterSection;
+
+/**
+ * URigMapperEditorSubsystem
+ * Subsystem to remap animation from/to all kinds of formats using the Rig Mapper API.
+ */
+UCLASS()
+class RIGMAPPEREDITOR_API URigMapperEditorSubsystem : public UEditorSubsystem
+{
+	GENERATED_BODY()
+
+	using FPose = TMap<FString, float>;
+	
+	using FPoseValues = TArray<float>;
+
+	using FSparseBakedCurves = TMap<FFrameTime, FPoseValues>;
+	
+public:
+	/**
+	 * Remap a CSV file using the given definitions, and optionally outputting separate files for each definition's output
+	 * The following header is expected at the first line of the CSV file: curve_name, frame_number, value
+	 */
+	UFUNCTION(BlueprintCallable, Category = "Editor Scripting|Animation|Rig Mapper")
+	static bool ConvertCsv(const FFilePath& InputFile, const FFilePath& OutputFile, const TArray<URigMapperDefinition*>& Definitions, bool bOutputIntermediateCsvFiles = false);
+
+	/**
+	 * Import a CSV file to an existing anim sequence, remapping it using the given definitions
+	 * The following header is expected at the first line of the CSV file: curve_name, frame_number, value
+	 * The target sequence rate will be unchanged (see SetAnimSequenceRate if required)
+	 */
+	UFUNCTION(BlueprintCallable, Category = "Editor Scripting|Animation|Rig Mapper")
+	static bool ConvertCsvToAnimSequence(const FFilePath& InputFile, UAnimSequence* Target, const TArray<URigMapperDefinition*>& Definitions);
+
+	/**
+	 * Import a CSV file to a new anim sequence asset, remapping it using the given definitions
+	 * The following header is expected at the first line of the CSV file: curve_name, frame_number, value
+	 */
+	UFUNCTION(BlueprintCallable, Category = "Editor Scripting|Animation|Rig Mapper")
+	static bool ConvertCsvToAnimSequenceNew(const FFilePath& InputFile, const FFilePath& NewAssetPath, USkeletalMesh* TargetMesh, const TArray<URigMapperDefinition*>& Definitions, const FFrameRate& FrameRate);
+
+	/**
+	 * Import a CSV file to the given control rig section, remapping it using the given definitions
+	 * The following header is expected at the first line of the CSV file: curve_name, frame_number, value
+	 * Control Rig transform/vector2d/float channels are currently supported
+	 * The target sequence rate will be unchanged
+	 */
+	UFUNCTION(BlueprintCallable, Category = "Editor Scripting|Animation|Rig Mapper")
+	static bool ConvertCsvToControlRigSection(const FFilePath& InputFile, UMovieSceneControlRigParameterSection* Target, const TArray<URigMapperDefinition*>& Definitions);
+
+	/**
+	 * Import a CSV file to a control rig section in a new Level Sqeuence, remapping it using the given definitions
+	 * The following header is expected at the first line of the CSV file: curve_name, frame_number, value
+	 * Control Rig transform/vector2d/float channels are currently supported
+	 */
+	UFUNCTION(BlueprintCallable, Category = "Editor Scripting|Animation|Rig Mapper")
+	static bool ConvertCsvToControlRigSectionNew(const FFilePath& InputFile, USkeletalMesh* TargetMesh, const TArray<URigMapperDefinition*>& Definitions, const FFilePath& NewAssetPath, const FFrameRate& FrameRate, const TSubclassOf<UControlRig>& ControlRigClass);
+
+	/**
+	 * Remap an anim sequence to another existing one, remapping it using the given definitions
+	 * The target sequence rate will be unchanged (see SetAnimSequenceRate if required)
+	 */
+	UFUNCTION(BlueprintCallable, Category = "Editor Scripting|Animation|Rig Mapper")
+	static bool ConvertAnimSequence(const UAnimSequence* Source, UAnimSequence* Target, const TArray<URigMapperDefinition*>& Definitions);
+
+	/**
+	 * Remap an anim sequence to a new one, remapping it using the given definitions
+	 */
+	UFUNCTION(BlueprintCallable, Category = "Editor Scripting|Animation|Rig Mapper")
+	static bool ConvertAnimSequenceNew(const UAnimSequence* Source, USkeletalMesh* TargetMesh, const TArray<URigMapperDefinition*>& Definitions);
+
+	/**
+	 * Export a CSV file from an anim sequence, remapping it using the given definitions
+	 */
+	UFUNCTION(BlueprintCallable, Category = "Editor Scripting|Animation|Rig Mapper")
+	static bool ConvertAnimSequenceToCsv(const UAnimSequence* Source, const FFilePath& OutputFile, const TArray<URigMapperDefinition*>& Definitions, bool bOutputIntermediateCsvFiles = false);
+
+	/**
+	 * Convert an AnimSequence to an existing Control Rig section, remapping it using the given definitions
+	 * Control Rig transform/vector2d/float channels are currently supported
+	 * The target sequence rate will be unchanged
+	 */
+	UFUNCTION(BlueprintCallable, Category = "Editor Scripting|Animation|Rig Mapper")
+	static bool ConvertAnimSequenceToControlRigSection(const UAnimSequence* Source, UMovieSceneControlRigParameterSection* Target, const TArray<URigMapperDefinition*>& Definitions);
+
+	/**
+	 * Convert an AnimSequence to a Control Rig section in a new Level Sequence, remapping it using the given definitions
+	 * Control Rig transform/vector2d/float channels are currently supported
+	 */
+	UFUNCTION(BlueprintCallable, Category = "Editor Scripting|Animation|Rig Mapper")
+	static bool ConvertAnimSequenceToControlRigSectionNew(const UAnimSequence* Source, USkeletalMesh* TargetMesh, const TArray<URigMapperDefinition*>& Definitions, const TSubclassOf<UControlRig>& ControlRigClass);
+
+	/**
+	 * Convert a Control Rig section to another existing one, remapping it using the given definitions
+	 * Control Rig transform/vector2d/float channels are currently supported
+	 * The target sequence rate will be unchanged
+	 */
+	UFUNCTION(BlueprintCallable, Category = "Editor Scripting|Animation|Rig Mapper")
+	static bool ConvertControlRigSection(const UMovieSceneControlRigParameterSection* Source, UMovieSceneControlRigParameterSection* Target, const TArray<URigMapperDefinition*>& Definitions);
+
+	/**
+	 * Convert a Control Rig section to a section in a new Level Sequence, remapping it using the given definitions
+	 * Control Rig transform/vector2d/float channels are currently supported
+	 */
+	UFUNCTION(BlueprintCallable, Category = "Editor Scripting|Animation|Rig Mapper")
+	static bool ConvertControlRigSectionNew(const UMovieSceneControlRigParameterSection* Source, USkeletalMesh* TargetMesh, const TArray<URigMapperDefinition*>& Definitions, const TSubclassOf<UControlRig>& ControlRigClass);
+
+	/**
+	 * Export a CSV file from a Control Rig Section, remapping it using the given definitions
+	 */
+	UFUNCTION(BlueprintCallable, Category = "Editor Scripting|Animation|Rig Mapper")
+	static bool ConvertControlRigSectionToCsv(const UMovieSceneControlRigParameterSection* Source, const FFilePath& OutputFile, const TArray<URigMapperDefinition*>& Definitions, bool bOutputIntermediateCsvFiles = false);
+
+	/**
+	 * Convert a Control Rig section to an existing AnimSequence, remapping it using the given definitions
+	 * The target sequence rate will be unchanged
+	 */
+	UFUNCTION(BlueprintCallable, Category = "Editor Scripting|Animation|Rig Mapper")
+	static bool ConvertControlRigSectionToAnimSequence(const UMovieSceneControlRigParameterSection* Source, UAnimSequence* Target, const TArray<URigMapperDefinition*>& Definitions);
+
+	/**
+	 * Convert a Control Rig section to a new AnimSequence asset, remapping it using the given definitions
+	 */
+	UFUNCTION(BlueprintCallable, Category = "Editor Scripting|Animation|Rig Mapper")
+	static bool ConvertControlRigSectionToAnimSequenceNew(const UMovieSceneControlRigParameterSection* Source, USkeletalMesh* TargetMesh, const TArray<URigMapperDefinition*>& Definitions);
+	
+	UFUNCTION(BlueprintCallable, Category = "Editor Scripting|Animation|Rig Mapper")
+	static ULevelSequence* GetSequenceFromSection(const UMovieSceneControlRigParameterSection* Section);
+
+	UFUNCTION(BlueprintCallable, Category = "Editor Scripting|Animation|Rig Mapper")
+	static TArray<UMovieSceneControlRigParameterSection*> GetSectionsFromSequence(ULevelSequence* Sequence);
+	
+	UFUNCTION(BlueprintCallable, Category = "Editor Scripting|Animation|Rig Mapper")
+	static FFrameRate GetAnimSequenceRate(const UAnimSequence* AnimSequence);
+
+	UFUNCTION(BlueprintCallable, Category = "Editor Scripting|Animation|Rig Mapper")
+	static void SetAnimSequenceRate(UAnimSequence* AnimSequence, FFrameRate FrameRate, bool bSetImportProperties);
+	
+private:
+	static bool ConvertCurveValuesToCsv(TArray<FString>& CurveNames, const TArray<FFrameTime>& FrameTimes, TArray<FPoseValues>& CurveValuesPerFrame, const FFilePath& OutputFile, const TArray<URigMapperDefinition*>& Definitions, bool bOutputIntermediateCsvFiles);
+
+	static bool WriteCurveValuesToCsv(const FFilePath& OutputFile, const TArray<FString>& CurveNames, const TArray<FFrameTime>& FrameTimes, const TArray<FPoseValues>& CurveValuesPerFrame);
+	
+	static bool LoadCurveValuesFromCsv(const FFilePath& InputFile, TArray<FString>& CurveNames, TArray<FFrameTime>& FrameTimes, TArray<FPoseValues>& CurveValuesPerFrame);
+	
+	static bool ConvertCurveValuesToAnimSequence(TArray<FString>& CurveNames, const TArray<FFrameTime>& FrameTimes, TArray<FPoseValues>& CurveValuesPerFrame, UAnimSequence* Target, const TArray<URigMapperDefinition*>& Definitions);
+
+	static bool AddCurveValuesToAnimSequence(UAnimSequence* Target, const TArray<FString>& CurveNames, const TArray<FFrameTime>& FrameTimes, const TArray<FPoseValues>& CurveValuesPerFrame);
+
+	static bool LoadCurveValuesFromAnimSequence(const UAnimSequence* Source, TArray<FString>& CurveNames, TArray<FFrameTime>& FrameTimes, TArray<FPoseValues>& CurveValuesPerFrame);
+	
+	static bool ConvertCurveValuesToControlRigSection(TArray<FString>& CurveNames, const TArray<FFrameTime>& FrameTimes, TArray<FPoseValues>& CurveValuesPerFrame, UMovieSceneControlRigParameterSection* Target, const TArray<URigMapperDefinition*>& Definitions);
+	
+	static bool AddCurveValuesToControlRigSection(UMovieSceneControlRigParameterSection* Target, const TArray<FString>& CurveNames, const TArray<FFrameTime>& FrameTimes, const TArray<FPoseValues>& CurveValuesPerFrame);
+
+	static bool LoadCurveValuesFromControlRigSection(const UMovieSceneControlRigParameterSection* Source, TArray<FString>& CurveNames, TArray<FFrameTime>& FrameTimes, TArray<FPoseValues>& CurveValuesPerFrame, const TArray<FString>& InputNames);
+	
+	static bool InitializeRigMapping(const TArray<URigMapperDefinition*>& Definitions, TArray<TSharedPtr<FacialRigMapping::FRigMapper>>& RigMappers);
+
+	static void EvaluateRigMapping(TArray<FString>& CurveNames, const TArray<FFrameTime>& FrameTimes, TArray<FPoseValues>& CurveValuesPerFrame, const TSharedPtr<FacialRigMapping::FRigMapper>& RigMapper);
+
+	static bool BakeSparseKeys(const TArray<FPose>& Poses, const TArray<FString>& CurveNames, TArray<FFrameTime>& FrameTimes, TArray<FPoseValues>& CurveValuesPerFrame);
+
+	static void SparseBakeCurve(const FString& CurveName, int32 CurveIndex, float CurveValue, TArray<FPoseValues>& CurveValuesPerFrame, const TArray<FFrameTime>& FrameTimesInOrder, int32 FrameIndex, const TArray<FPose>& Poses, const TArray<FFrameTime>& FrameTimes, const int32 ActualFrameIndex);
+};
diff --git a/Engine/Plugins/Experimental/Animation/RigMapper/Source/RigMapperEditor/RigMapperEditor.Build.cs b/Engine/Plugins/Experimental/Animation/RigMapper/Source/RigMapperEditor/RigMapperEditor.Build.cs
new file mode 100644
--- /dev/null
+++ b/Engine/Plugins/Experimental/Animation/RigMapper/Source/RigMapperEditor/RigMapperEditor.Build.cs
@@ -0,0 +1,41 @@
+// Copyright Epic Games, Inc. All Rights Reserved.
+
+using UnrealBuildTool;
+
+public class RigMapperEditor : ModuleRules
+{
+	public RigMapperEditor(ReadOnlyTargetRules Target) : base(Target)
+	{
+		PublicDependencyModuleNames.AddRange(
+			new string[]
+			{
+				"EditorSubsystem",
+				"UnrealEd",
+				"GraphEditor",
+				"Engine"
+			});
+
+		PrivateDependencyModuleNames.AddRange(
+			new string[]
+			{
+				"Core",
+				"CoreUObject",
+				"AnimGraph",
+				"BlueprintGraph",
+				"RigMapper",
+				"ControlRig",
+				"LevelSequence",
+				"MovieScene",
+				"MovieSceneTracks",
+				"TargetPlatform",
+				"AssetTools",
+				"Slate",
+				"SlateCore",
+				"ToolMenus",
+				"ContentBrowser",
+				"MessageLog",
+				"InputCore",
+				"Sequencer"
+			});
+	}
+}
