Rig Mapper plugin update:
- Optimizations & performance improvements
- Refactoring & Commenting

diff --git a/Engine/Plugins/Experimental/Animation/RigMapper/Source/RigMapper/Private/AnimNode_RigMapper.cpp b/Engine/Plugins/Experimental/Animation/RigMapper/Source/RigMapper/Private/AnimNode_RigMapper.cpp
--- a/Engine/Plugins/Experimental/Animation/RigMapper/Source/RigMapper/Private/AnimNode_RigMapper.cpp
+++ b/Engine/Plugins/Experimental/Animation/RigMapper/Source/RigMapper/Private/AnimNode_RigMapper.cpp
@@ -4,13 +4,15 @@
 
 #include "UObject/SoftObjectPath.h"
 #include "RigMapperDefinition.h"
+#include "Animation/AnimCurveUtils.h"
 #include "Engine/SkeletalMesh.h"
 #include "Animation/AnimInstanceProxy.h"
 
 DEFINE_STAT(STAT_RigMapperEvaluateNode);
 DEFINE_STAT(STAT_RigMapperEvaluateRig);
 DEFINE_STAT(STAT_RigMapperEvaluateAllRigs);
-DEFINE_STAT(STAT_RigMapperUpdateValues);
+DEFINE_STAT(STAT_RigMapperInitialize);
+DEFINE_STAT(STAT_RigMapperSetCurves);
 
 FAnimNode_RigMapper::FAnimNode_RigMapper()
 {
@@ -56,16 +58,20 @@ void FAnimNode_RigMapper::PreUpdate(const UAnimInstance* InAnimInstance)
 {
 	FAnimNode_Base::PreUpdate(InAnimInstance);
 	
+	// Need reinit if the definitions in use have changed.
+	// todo: check against definition property changes? 
 	bool bReInit = false;
-	TArray<TObjectPtr<URigMapperDefinition>>* DefinitionsToCheck = &Definitions;
 
+	// If definitions were loaded from the SKM asset user data, they take priority and are the ones we will check against
+	TArray<TObjectPtr<URigMapperDefinition>>* DefinitionsToCheck = &Definitions;
 	if (LoadedUserData)
 	{
 		DefinitionsToCheck = &LoadedUserData->Definitions;
 	}
 
 	check(DefinitionsToCheck);
-	
+
+	// Otherwise, reinit if the definitions are not the same
 	if (DefinitionsToCheck->Num() != LoadedDefinitions.Num())
 	{
 		bReInit = true;
@@ -139,64 +145,83 @@ void FAnimNode_RigMapper::Evaluate_AnyThread(FPoseContext& Output)
 void FAnimNode_RigMapper::EvaluateRigMapping(FPoseContext& Output) const
 {
 	SCOPE_CYCLE_COUNTER(STAT_RigMapperEvaluateAllRigs);
+
+	// We need tp cache current curve values for possible Lerp using given alpha
+	TArray<double> InputValues;
 	
-	for (int i = 0; i < RigMappers.Num(); i++)
+	for (int32 RigMapperIndex = 0; RigMapperIndex < RigMappers.Num(); RigMapperIndex++)
 	{
-		const TSharedPtr<FacialRigMapping::FRigMapper> RigMapper = RigMappers[i];
+		const TSharedPtr<FacialRigMapping::FRigMapper> RigMapper = RigMappers[RigMapperIndex];
 		RigMapper->SetDirty();
-
+		
 		// Set first rig mapper's input values from current pose
-		if (i == 0)
+		if (RigMapperIndex == 0)
 		{
-			SCOPE_CYCLE_COUNTER(STAT_RigMapperUpdateValues);
-			
-			const TArray<FString>& InputNames = RigMapper->GetInputNames();
-			for (int n = 0; n < InputNames.Num(); n++)
+			// Retrieve values from current pose
+			InputValues.AddZeroed(InputCurveMappings.Num());
+			UE::Anim::FCurveUtils::BulkGet(Output.Curve, InputCurveMappings, 
+	[&InputValues](const FRigMapperCurveMapping& InBulkElement, const float InValue)
+				{
+					InputValues[InBulkElement.CurveIndex] = InValue;
+				}
+			);
+
+			// Set first rig mapper's input values from current pose
+			for (int32 InputIndex = 0; InputIndex < InputValues.Num(); InputIndex++)
 			{
-				RigMapper->SetDirectValue(n, Output.Curve.Get(*InputNames[n]));
+				RigMapper->SetDirectValue(InputIndex, InputValues[InputIndex]);
 			}
 		}
 		else
 		{
 			// Other RigMappers: Link/Map prev rig mapper outputs to current inputs
-			check(IndexMaps.IsValidIndex(i - 1))
-			const FRigMapperIndexMap& IndexMapping = IndexMaps[i - 1];
+			check(IndexMaps.IsValidIndex(RigMapperIndex - 1))
+			const FRigMapperIndexMap& IndexMapping = IndexMaps[RigMapperIndex - 1];
 
 			// Get output values from previous rig mapper
 			TArray<double> PrevOutputValuesInOrder;
 			{
 				SCOPE_CYCLE_COUNTER(STAT_RigMapperEvaluateRig);
-				PrevOutputValuesInOrder = RigMappers[i - 1]->GetOutputValuesInOrder();
+				PrevOutputValuesInOrder = RigMappers[RigMapperIndex - 1]->GetOutputValuesInOrder();
 			}
 
 			// Set input values for current rig mapper
+			for (int32 InputIndex = 0; InputIndex < IndexMapping.Num(); InputIndex++)
 			{
-				SCOPE_CYCLE_COUNTER(STAT_RigMapperUpdateValues);
-				
-				for (int32 n = 0; n < IndexMapping.Num(); n++)
+				if (const int32 OutputIndex = IndexMapping[InputIndex]; OutputIndex != INDEX_NONE)
 				{
-					if (const int32 Index = IndexMapping[n]; Index != INDEX_NONE)
-					{
-						RigMapper->SetDirectValue(n, PrevOutputValuesInOrder[Index]);
-					}
+					RigMapper->SetDirectValue(InputIndex, PrevOutputValuesInOrder[OutputIndex]);
 				}
 			}
 		}
 	}
-
+	
 	// Set current pose from last rig mapper's output values 
 	if (!RigMappers.IsEmpty())
 	{
-		TMap<FString, double> OutputValues;
+		// Retrieve last rig mapper output values
+		TArray<double> OutputValues;
 		{
 			SCOPE_CYCLE_COUNTER(STAT_RigMapperEvaluateRig);
-			OutputValues = RigMappers.Last()->GetOutputValues();
+			OutputValues = RigMappers.Last()->GetOutputValuesInOrder();
 		}
-		
-		for (const TPair<FString, double>& OutputValue : OutputValues)
+
 		{
-			const float Value = FMath::Lerp(Output.Curve.Get(*OutputValue.Key), OutputValue.Value, Alpha);
-			Output.Curve.Set(*OutputValue.Key, Value);
+			SCOPE_CYCLE_COUNTER(STAT_RigMapperSetCurves);
+
+			// Set all output curves for current pose from the given output values. Possibly lerping with current values depending on Alpha.
+			UE::Anim::FCurveUtils::BulkSet(Output.Curve, OutputCurveMappings, 
+						[&InputValues, &OutputValues, LerpAlpha=Alpha](const FRigMapperOutputCurveMapping& InBulkElement)
+				{
+					float Value = OutputValues[InBulkElement.CurveIndex];
+							
+					if (InputValues.IsValidIndex(InBulkElement.InputCurveIndex))
+					{
+						Value = FMath::Lerp(InputValues[InBulkElement.InputCurveIndex], Value, LerpAlpha);
+					}
+					return Value;
+				}
+			);
 		}
 	}
 }
@@ -208,8 +233,11 @@ int32 FAnimNode_RigMapper::GetLODThreshold() const
 
 bool FAnimNode_RigMapper::InitializeRigMapping(USkeletalMesh* TargetMesh)
 {
+	SCOPE_CYCLE_COUNTER(STAT_RigMapperInitialize);
+	
 	TArray<URigMapperDefinition*> DefinitionsToLoad;
 
+	// Retrieve the definitions to load and use (either the ones set on the node, or the ones overriden by the SKM asset user data)
 	if (TargetMesh)
 	{
 		LoadedUserData = Cast<URigMapperDefinitionUserData>(TargetMesh->GetAssetUserDataOfClass(URigMapperDefinitionUserData::StaticClass()));
@@ -235,6 +263,8 @@ bool FAnimNode_RigMapper::InitializeRigMapping(USkeletalMesh* TargetMesh)
 			return false;
 		}
 	}
+
+	// Initialize the actual RigMappers from the definitions to load
 	RigMappers.Empty(DefinitionsToLoad.Num());
 	IndexMaps.Empty(DefinitionsToLoad.Num() - 1);
 	
@@ -257,6 +287,28 @@ bool FAnimNode_RigMapper::InitializeRigMapping(USkeletalMesh* TargetMesh)
 		RigMappers.Add(RigMapper);
 	}
 
+	// Cache a map of curve indices to bulk get current curve values for the current pose
+	const TArray<FString>& InputNames = RigMappers[0]->GetInputNames();
+	InputCurveMappings.Empty();
+	InputCurveMappings.Reserve(InputNames.Num());
+	for (int32 InputIndex = 0; InputIndex < InputNames.Num(); InputIndex++)
+	{
+		InputCurveMappings.Add(*InputNames[InputIndex], InputIndex);
+	}
+
+	// Cache a map of curve indices to bulk set the new curve values for the output pose
+	const TArray<FString>& OutputNames = RigMappers.Last()->GetOutputNames();
+	OutputCurveMappings.Empty();
+	OutputCurveMappings.Reserve(OutputNames.Num());
+	for (int32 OutputIndex = 0; OutputIndex < OutputNames.Num(); OutputIndex++)
+	{
+		const FString& CurveName = OutputNames[OutputIndex];
+
+		// We need to cache the matching input curve index to allow lerping depending on the node's Alpha
+		int32 InputIndex = InputNames[OutputIndex] == CurveName ? OutputIndex : InputNames.Find(CurveName);
+		OutputCurveMappings.Add(*OutputNames[OutputIndex], OutputIndex, InputIndex);
+	}
+	
 	LoadedDefinitions = MoveTemp(DefinitionsToLoad);
 	return true;
 }
@@ -265,11 +317,10 @@ FRigMapperIndexMap FAnimNode_RigMapper::MakeIndexMap(const TSharedPtr<FacialRigM
 {
 	const TArray<FString> FromOutputs = OutputRig->GetOutputNames();
 	const TArray<FString> ToInputs = InputRig->GetInputNames();
-				
+	
 	FRigMapperIndexMap IndexMapping;
 	
 	IndexMapping.Reserve(ToInputs.Num());
-
 	for (const FString& Input : ToInputs)
 	{
 		IndexMapping.Add(FromOutputs.Find(Input));
diff --git a/Engine/Plugins/Experimental/Animation/RigMapper/Source/RigMapper/Public/AnimNode_RigMapper.h b/Engine/Plugins/Experimental/Animation/RigMapper/Source/RigMapper/Public/AnimNode_RigMapper.h
--- a/Engine/Plugins/Experimental/Animation/RigMapper/Source/RigMapper/Public/AnimNode_RigMapper.h
+++ b/Engine/Plugins/Experimental/Animation/RigMapper/Source/RigMapper/Public/AnimNode_RigMapper.h
@@ -18,10 +18,11 @@ class USkeletalMesh;
 
 DECLARE_STATS_GROUP(TEXT("RigMapper"), STATGROUP_RigMapper, STATCAT_Advanced);
 
-DECLARE_CYCLE_STAT_EXTERN(TEXT("Evaluate AnimNode"), STAT_RigMapperEvaluateNode, STATGROUP_RigMapper, RIGMAPPER_API);
-DECLARE_CYCLE_STAT_EXTERN(TEXT("Evaluate RigMapping"), STAT_RigMapperEvaluateAllRigs, STATGROUP_RigMapper, RIGMAPPER_API);
-DECLARE_CYCLE_STAT_EXTERN(TEXT("Evaluate RigMapper"), STAT_RigMapperEvaluateRig, STATGROUP_RigMapper, RIGMAPPER_API);
-DECLARE_CYCLE_STAT_EXTERN(TEXT("Update RigMapper"), STAT_RigMapperUpdateValues, STATGROUP_RigMapper, RIGMAPPER_API);
+DECLARE_CYCLE_STAT_EXTERN(TEXT("Evaluate Rig Mapper Node"), STAT_RigMapperEvaluateNode, STATGROUP_RigMapper, RIGMAPPER_API);
+DECLARE_CYCLE_STAT_EXTERN(TEXT("Evaluate Rig Mapping"), STAT_RigMapperEvaluateAllRigs, STATGROUP_RigMapper, RIGMAPPER_API);
+DECLARE_CYCLE_STAT_EXTERN(TEXT("Evaluate Rig Mapper Definition"), STAT_RigMapperEvaluateRig, STATGROUP_RigMapper, RIGMAPPER_API);
+DECLARE_CYCLE_STAT_EXTERN(TEXT("Initialize Rig Mapping"), STAT_RigMapperInitialize, STATGROUP_RigMapper, RIGMAPPER_API);
+DECLARE_CYCLE_STAT_EXTERN(TEXT("Rig Mapper Set Pose Curves"), STAT_RigMapperSetCurves, STATGROUP_RigMapper, RIGMAPPER_API);
 
 using FRigMapperIndexMap = TArray<int32>;
 
@@ -46,12 +47,15 @@ public:
 	virtual void Update_AnyThread(const FAnimationUpdateContext& Context) override;
 	virtual void CacheBones_AnyThread(const FAnimationCacheBonesContext& Context) override;
 	virtual void Evaluate_AnyThread(FPoseContext & Output) override;
-	void EvaluateRigMapping(FPoseContext& Output) const;
 	virtual int32 GetLODThreshold() const override;
 	// End of FAnimNode_Base interface
 
+	// Initialize the rig mapper(s) from the current definitions or SKM asset user data and cache the curve indices & index mappings needed for evaluation
 	bool InitializeRigMapping(USkeletalMesh* TargetMesh=nullptr);
 
+	// Evaluate the new curve values using the initialized rig mappers and sets the new output pose 
+	void EvaluateRigMapping(FPoseContext& Output) const;
+
 public:
 	UPROPERTY(EditAnywhere, BlueprintReadWrite, Category = Links)
 	FPoseLink SourcePose;
@@ -71,12 +75,55 @@ public:
 	UPROPERTY(EditAnywhere, BlueprintReadWrite, Category = Alpha, meta = (PinShownByDefault))
 	float Alpha = 1.0f;
 
-private:	
+private:
+	// Builds an In/Out index mapping between 2 rig mapper definitions chained together
 	static FRigMapperIndexMap MakeIndexMap(const TSharedPtr<FacialRigMapping::FRigMapper>& OutputRig, const TSharedPtr<FacialRigMapping::FRigMapper>& InputRig);
 	
 private:
+	// The inititalized rig mappers to evaluate the remapped pose
 	TArray<TSharedPtr<FacialRigMapping::FRigMapper>> RigMappers;
+
+	// The definitions that we have loaded. Cached to check against changes and reinit if need be
 	TArray<TObjectPtr<URigMapperDefinition>> LoadedDefinitions;
+
+	// The asset user data currently used to override definitions, if any was set on the skeletal mesh
 	TObjectPtr<URigMapperDefinitionUserData> LoadedUserData;
+
+	// Mapping between rig mappers inputs and outputs in order
 	TArray<FRigMapperIndexMap> IndexMaps;
+
+	// Base curve mapping for bulk get/set of the linked pose curves
+	struct FRigMapperCurveMapping
+	{
+		FRigMapperCurveMapping() = default;
+		
+		FRigMapperCurveMapping(FName InName, int32 InCurveIndex)
+			: Name(InName)
+			, CurveIndex(InCurveIndex)
+		{}
+
+		FName Name = NAME_None;
+		int32 CurveIndex = INDEX_NONE;
+	};
+
+	// Curve mapping for bulk set of the linked pose curves with a mapping to the matching input to allow lerping with the current curve value
+	struct FRigMapperOutputCurveMapping : FRigMapperCurveMapping
+	{
+		FRigMapperOutputCurveMapping() = default;
+		
+		FRigMapperOutputCurveMapping(FName InName, int32 InOutputCurveIndex, int32 InInputCurveIndex)
+			: FRigMapperCurveMapping(InName, InOutputCurveIndex)
+			, InputCurveIndex(InInputCurveIndex)
+		{}
+		
+		int32 InputCurveIndex = INDEX_NONE;
+	};
+
+	// Input curve mapping for bulk get of the linked pose curve values
+	using FInputCurveMappings = UE::Anim::TNamedValueArray<FDefaultAllocator, FRigMapperCurveMapping>;
+	FInputCurveMappings InputCurveMappings;
+	
+	// Output curve mapping for the bulk get of the linked pose curve values
+	using FOutputCurveMappings = UE::Anim::TNamedValueArray<FDefaultAllocator, FRigMapperOutputCurveMapping>;
+	FOutputCurveMappings OutputCurveMappings;
 };
