Fix rig mapper crash when changing definitions at runtime if definition has mismatch in inputs and outputs count.

diff --git a/Engine/Plugins/Experimental/Animation/RigMapper/Source/RigMapper/Private/AnimNode_RigMapper.cpp b/Engine/Plugins/Experimental/Animation/RigMapper/Source/RigMapper/Private/AnimNode_RigMapper.cpp
--- a/Engine/Plugins/Experimental/Animation/RigMapper/Source/RigMapper/Private/AnimNode_RigMapper.cpp
+++ b/Engine/Plugins/Experimental/Animation/RigMapper/Source/RigMapper/Private/AnimNode_RigMapper.cpp
@@ -305,7 +305,7 @@ bool FAnimNode_RigMapper::InitializeRigMapping(USkeletalMesh* TargetMesh)
 		const FString& CurveName = OutputNames[OutputIndex];
 
 		// We need to cache the matching input curve index to allow lerping depending on the node's Alpha
-		int32 InputIndex = InputNames[OutputIndex] == CurveName ? OutputIndex : InputNames.Find(CurveName);
+		int32 InputIndex = InputNames.Find(CurveName);
 		OutputCurveMappings.Add(*OutputNames[OutputIndex], OutputIndex, InputIndex);
 	}
 	
