One liner to fix validation in RigMapperLinkedDefinitions not displaying a popup for a certain validation error

diff --git a/Engine/Plugins/Experimental/Animation/RigMapper/Source/RigMapper/Private/RigMapperDefinition.cpp b/Engine/Plugins/Experimental/Animation/RigMapper/Source/RigMapper/Private/RigMapperDefinition.cpp
--- a/Engine/Plugins/Experimental/Animation/RigMapper/Source/RigMapper/Private/RigMapperDefinition.cpp
+++ b/Engine/Plugins/Experimental/Animation/RigMapper/Source/RigMapper/Private/RigMapperDefinition.cpp
@@ -1141,6 +1141,7 @@ bool URigMapperLinkedDefinitions::AreLinkedDefinitionsValid() const
 				if (!SourceDefinitions[DefinitionIndex]->Inputs.Contains(Output.Key))
 				{
 					UE_LOG(LogRigMapperDefinition, Warning, TEXT("Could not find matching input in definition %d for output %s in definition %d"), DefinitionIndex, *Output.Key, DefinitionIndex - 1)
+					bOk = false;
 				}
 			}
 		}
