Persist Json load and save directories in the EditorPerProjectIni file

diff --git a/Engine/Plugins/Experimental/Animation/RigMapper/Source/RigMapperEditor/Private/RigMapperEditorModule.cpp b/Engine/Plugins/Experimental/Animation/RigMapper/Source/RigMapperEditor/Private/RigMapperEditorModule.cpp
--- a/Engine/Plugins/Experimental/Animation/RigMapper/Source/RigMapperEditor/Private/RigMapperEditorModule.cpp
+++ b/Engine/Plugins/Experimental/Animation/RigMapper/Source/RigMapperEditor/Private/RigMapperEditorModule.cpp
@@ -23,6 +23,8 @@
 #include "ToolMenus.h"
 #include "Logging/MessageLog.h"
 #include "RigMapperGraph/RigMapperDefinitionEditorGraphNode.h"
+#include "Misc/ConfigCacheIni.h"
+#include "Misc/Paths.h"
 
 const FName FRigMapperEditorModule::AppIdentifier("RigMapperEditorApp");
 
@@ -118,16 +120,16 @@ void FRigMapperEditorModule::RegisterRigMapperDefinitionContextMenuEntries()
 		"RigMapperDefinition_ExportToJson",
 		LOCTEXT("RigMapperDefinitionContextMenu_ExportToJson_Label", "Export to Json"),
 		LOCTEXT("RigMapperDefinitionContextMenu_ExportToJson_Tooltip", "Export this definition as a json file"),
-		FSlateIcon(FAppStyle::GetAppStyleSetName(), "Icons.Toolbar.Export"),
-		FToolMenuExecuteAction::CreateStatic(&FRigMapperEditorModule::ExportToJson, false)
-	);
-	Section.AddMenuEntry(
-		"RigMapperDefinition_Validate",
-		LOCTEXT("RigMapperDefinitionContextMenu_Validate_Label", "Validate Definition"),
-		LOCTEXT("RigMapperDefinitionContextMenu_Validate_Tooltip", "Check if the definition has any noticable issue"),
-		FSlateIcon(FAppStyle::GetAppStyleSetName(), "Icons.Validate"),
-		FToolMenuExecuteAction::CreateStatic(&FRigMapperEditorModule::ValidateDefinition, false)
-	);
+			FSlateIcon(FAppStyle::GetAppStyleSetName(), "Icons.Toolbar.Export"),
+			FToolMenuExecuteAction::CreateStatic(&FRigMapperEditorModule::ExportToJson, false)
+			);
+			Section.AddMenuEntry(
+				"RigMapperDefinition_Validate",
+				LOCTEXT("RigMapperDefinitionContextMenu_Validate_Label", "Validate Definition"),
+				LOCTEXT("RigMapperDefinitionContextMenu_Validate_Tooltip", "Check if the definition has any noticable issue"),
+				FSlateIcon(FAppStyle::GetAppStyleSetName(), "Icons.Validate"),
+				FToolMenuExecuteAction::CreateStatic(&FRigMapperEditorModule::ValidateDefinition, false)
+			);
 }
 
 void FRigMapperEditorModule::RegisterRigMapperDefinitionToolbarEntries()
@@ -165,7 +167,7 @@ void FRigMapperEditorModule::RegisterRigMapperLinkedDefinitionsContextMenuEntrie
 
 	const FToolMenuInsert SectionInsertLocation("GetAssetActions", EToolMenuInsertType::After);
 	FToolMenuSection& Section = Menu->AddSection("RigMapperActions", INVTEXT("Rig Mapper"), SectionInsertLocation);
-	
+
 	Section.AddMenuEntry(
 		"RigMapperLinkedDefinitions_Bake",
 		LOCTEXT("RigMapperLinkedDefinitionsContextMenu_Bake_Label", "Bake Definitions"),
@@ -188,7 +190,7 @@ void FRigMapperEditorModule::RegisterRigMapperLinkedDefinitionsToolbarEntries()
 	FToolMenuSection& Section = Menu->FindOrAddSection("AssetEditorActions");
 
 	// todo: bake to new
-	
+
 	Section.AddEntry(FToolMenuEntry::InitToolBarButton(
 		"RigMapperLinkedDefinitions_Bake",
 		FToolMenuExecuteAction::CreateStatic(&FRigMapperEditorModule::BakeDefinitions, true),
@@ -210,15 +212,21 @@ bool FRigMapperEditorModule::GetUserPickedPath(FString& OutPath, bool bImport)
 {
 	TArray<FString> Filenames;
 	IDesktopPlatform* DesktopPlatform = FDesktopPlatformModule::Get();
-	
+
 	bool bOpen;
-	
+
 	if (bImport)
 	{
+		FString LastRigMapperImportDirectory;
+		if (!GConfig->GetString(TEXT("RigMapper"), TEXT("LastImportDirectory"), LastRigMapperImportDirectory, GEditorPerProjectIni))
+		{
+			LastRigMapperImportDirectory = FEditorDirectories::Get().GetLastDirectory(ELastDirectory::GENERIC_IMPORT);
+		}
+
 		bOpen = DesktopPlatform->OpenFileDialog(
 			FSlateApplication::Get().FindBestParentWindowHandleForDialogs(nullptr),
 			NSLOCTEXT("RigMapperDefinitionIO", "ImportFromJson", "Import Definition from Json...").ToString(), 
-			FEditorDirectories::Get().GetLastDirectory(ELastDirectory::GENERIC_IMPORT),
+			LastRigMapperImportDirectory,
 			TEXT(""), 
 			TEXT("json (*.json)|*.json|"),
 			EFileDialogFlags::None,
@@ -228,20 +236,26 @@ bool FRigMapperEditorModule::GetUserPickedPath(FString& OutPath, bool bImport)
 	else
 	{
 		constexpr bool bDirectory = true;
+
+		FString LastRigMapperExportDirectory;
+		if (!GConfig->GetString(TEXT("RigMapper"), TEXT("LastExportDirectory"), LastRigMapperExportDirectory, GEditorPerProjectIni))
+		{
+			LastRigMapperExportDirectory = FEditorDirectories::Get().GetLastDirectory(ELastDirectory::GENERIC_EXPORT);
+		}
 		
 		if (bDirectory)
 		{
 			return DesktopPlatform->OpenDirectoryDialog(
 				FSlateApplication::Get().FindBestParentWindowHandleForDialogs(nullptr),
 				NSLOCTEXT("RigMapperDefinitionIO", "ExportToJson", "Export Definition to Json...").ToString(), 
-				FEditorDirectories::Get().GetLastDirectory(ELastDirectory::GENERIC_EXPORT),
+				LastRigMapperExportDirectory,
 				OutPath
 			);
 		}
 		bOpen = DesktopPlatform->SaveFileDialog(
 			FSlateApplication::Get().FindBestParentWindowHandleForDialogs(nullptr),
 			NSLOCTEXT("RigMapperDefinitionIO", "ExportToJson", "Export Definition to Json...").ToString(), 
-			FEditorDirectories::Get().GetLastDirectory(ELastDirectory::GENERIC_EXPORT),
+			LastRigMapperExportDirectory,
 			TEXT(""), 
 			TEXT("json (*.json)|*.json|"),
 			EFileDialogFlags::None,
@@ -262,6 +276,8 @@ void FRigMapperEditorModule::LoadFromJson(const FToolMenuContext& ToolMenuContex
 	const TArray<URigMapperDefinition*> Definitions = GetDefinitionsFromContext(ToolMenuContext, bFromAssetEditor);
 	
 	FFilePath JsonFile;
+	bool bSucceeded = false;
+	FString LastPath;
 
 	if (!Definitions.IsEmpty() && GetUserPickedPath(JsonFile.FilePath))
 	{
@@ -277,8 +293,21 @@ void FRigMapperEditorModule::LoadFromJson(const FToolMenuContext& ToolMenuContex
 				MessageLog.Error(FText::Format(LOCTEXT("RigMapperEditorActions_ImportFailed", "Failed to load definition \"{ObjectName}\" from Json. See output log for more details"), Arguments));
 				MessageLog.Open(EMessageSeverity::Error);
 			}
+			else
+			{
+				bSucceeded = true;
+				LastPath = JsonFile.FilePath;
+			}
 		}
 	}
+
+	if (bSucceeded)
+	{
+		// set the new import directory from the last successful save
+		FString NewLastRigMapperImportDirectory = FPaths::GetPath(LastPath);
+		GConfig->SetString(TEXT("RigMapper"), TEXT("LastImportDirectory"), *NewLastRigMapperImportDirectory, GEditorPerProjectIni);
+		GConfig->Flush(false, GEditorPerProjectIni);
+	}
 }
 
 void FRigMapperEditorModule::ExportToJson(const FToolMenuContext& ToolMenuContext, bool bFromAssetEditor)
@@ -286,6 +315,8 @@ void FRigMapperEditorModule::ExportToJson(const FToolMenuContext& ToolMenuContex
 	const TArray<URigMapperDefinition*> Definitions = GetDefinitionsFromContext(ToolMenuContext, bFromAssetEditor);
 	
 	FDirectoryPath Path;
+	bool bSucceeded = false;
+	FString LastPath;
 	
 	if (!Definitions.IsEmpty() && GetUserPickedPath(Path.Path, false))
 	{
@@ -304,8 +335,21 @@ void FRigMapperEditorModule::ExportToJson(const FToolMenuContext& ToolMenuContex
 				MessageLog.Error(FText::Format(LOCTEXT("RigMapperEditorActions_ExportFailed", "Failed to export definition \"{ObjectName}\". See output log for more details"), Arguments));
 				MessageLog.Open(EMessageSeverity::Error);
 			}
+			else
+			{
+				bSucceeded = true;
+				LastPath = ActualFilePath.FilePath;
+			}
 		}
 	}
+
+	if (bSucceeded)
+	{
+		// set the new export directory
+		FString NewLastRigMapperExportDirectory = FPaths::GetPath(LastPath);
+		GConfig->SetString(TEXT("RigMapper"), TEXT("LastExportDirectory"), *NewLastRigMapperExportDirectory, GEditorPerProjectIni);
+		GConfig->Flush(false, GEditorPerProjectIni);
+	}
 }
 
 void FRigMapperEditorModule::ValidateDefinition(const FToolMenuContext& ToolMenuContext, bool bFromAssetEditor)
