bugfix for validating URigMapperLinkedDefinitions more strictly revealed error in unit test data; fix this

diff --git a/Engine/Plugins/Experimental/Animation/RigMapper/Source/RigMapperDeveloper/Private/Tests/RigMapperTests.cpp b/Engine/Plugins/Experimental/Animation/RigMapper/Source/RigMapperDeveloper/Private/Tests/RigMapperTests.cpp
--- a/Engine/Plugins/Experimental/Animation/RigMapper/Source/RigMapperDeveloper/Private/Tests/RigMapperTests.cpp
+++ b/Engine/Plugins/Experimental/Animation/RigMapper/Source/RigMapperDeveloper/Private/Tests/RigMapperTests.cpp
@@ -110,7 +110,7 @@ TArray<URigMapperDefinition*> CreateValidRigMapperDefinitions2()
 	TArray<URigMapperDefinition*> Definitions = CreateValidRigMapperDefinitions1();
 
 	// modify the second definition so it uses an Null Output from the 1st stage
-	Definitions[1]->Inputs = { FString{ TEXT("OutputVal1")}, FString{ TEXT("OutputVal2")}, FString{ TEXT("OutputVal8")}, // OutputVal8 is a Null Output from the previous stage
+	Definitions[1]->Inputs = { FString{ TEXT("OutputVal1")}, FString{ TEXT("OutputVal2")}, FString{ TEXT("OutputVal3") }, FString{ TEXT("OutputVal8")}, // OutputVal8 is a Null Output from the previous stage
 		FString{ TEXT("OutputVal4")}, FString{ TEXT("OutputVal5")}, FString{ TEXT("OutputVal6")} };
 	TMap<FString, double> Inputs;
 	Inputs.Add(FString(TEXT("OutputVal1")), 0.2f);
