added additional test to ensure that baked and unbaked RigMapper plugin defs do not get out of sync with each other

diff --git a/Engine/Plugins/Experimental/Animation/RigMapper/Source/RigMapperDeveloper/Private/Tests/RigMapperTests.cpp b/Engine/Plugins/Experimental/Animation/RigMapper/Source/RigMapperDeveloper/Private/Tests/RigMapperTests.cpp
--- a/Engine/Plugins/Experimental/Animation/RigMapper/Source/RigMapperDeveloper/Private/Tests/RigMapperTests.cpp
+++ b/Engine/Plugins/Experimental/Animation/RigMapper/Source/RigMapperDeveloper/Private/Tests/RigMapperTests.cpp
@@ -817,7 +817,6 @@ bool FTestRigMapperTestRawDefinitionsRoundTrip::RunTest(const FString& InParamet
 	FString HighToLowDefPath;
 	FString ControlSubStr;
 	FString CurveIdentifier;
-	TArray<FString> AllowedMissingOutputControls;
 
 	const double Tolerance = 0.00001; 
 
@@ -965,7 +964,7 @@ bool FTestRigMapperTestRawDefinitionsRoundTrip::RunTest(const FString& InParamet
 			EvaluateRigMapping(CurveDataRigMapperInputs[Frame], IndexMaps, RigMappers, Outputs);
 
 			// check that the round trip gives outputs which are identical to the inputs
-			if (!AreMapsEqualWithTolerance(CurveDataRigMapperInputs[Frame], Outputs, CurveIdentifier, AllowedMissingOutputControls, Tolerance))
+			if (!AreMapsEqualWithTolerance(CurveDataRigMapperInputs[Frame], Outputs, CurveIdentifier, {}, Tolerance))
 			{
 				UE_LOG(LogRigMapperTest, Error, TEXT("Frame %d contains different values for LowToHigh->HighToLow round trip test"), Frame);
 				bResult = false;
@@ -981,4 +980,153 @@ bool FTestRigMapperTestRawDefinitionsRoundTrip::RunTest(const FString& InParamet
 	return bResult;
 }
 
+IMPLEMENT_COMPLEX_AUTOMATION_TEST(FTestRigMapperCompareBakedVsUnbakedPluginDefinitions, "RigMapper.RigMapperCompareBakedVsUnbakedPluginDefinitions", EAutomationTestFlags::EditorContext | EAutomationTestFlags::EngineFilter)
+
+void FTestRigMapperCompareBakedVsUnbakedPluginDefinitions::GetTests(TArray<FString>& OutBeautifiedNames, TArray<FString>& OutTestCommands) const
+{
+	OutBeautifiedNames.Add("Verify that baked definition RM_MHL_FNL gives the same result as the stack of individual unbaked definitions");
+	OutTestCommands.Add("MHL_FNL");
+
+	OutBeautifiedNames.Add("Verify that baked definition RM_MHL_CDL gives the same result as the stack of individual unbaked definitions");
+	OutTestCommands.Add("MHL_CDL");
+
+	OutBeautifiedNames.Add("Verify that baked definition RM_CDL_FNL gives the same result as the stack of individual unbaked definitions");
+	OutTestCommands.Add("CDL_FNL");
+}
+
+bool FTestRigMapperCompareBakedVsUnbakedPluginDefinitions::RunTest(const FString& InParameters)
+{
+	bool bResult = true;
+	UAnimSequence* TestAnimSequence = nullptr;
+	TArray<FString> UnbakedDefinitionPaths;
+	FString TestAnimSequencePath;
+	FString BakedDefinitionPath;
+	const FString ControlSubStr = TEXT("CTRL_");
+	const FString CurveIdentifier = TEXT("CTRL_expressions");
+
+	const double Tolerance = 0.00001;
+
+	// load in the appropriate Rig Mapper Definitions for the case we are testing
+	if (InParameters == "MHL_FNL")
+	{
+		BakedDefinitionPath = TEXT("/RigMapper/Definitions/Baked/RM_MHL_FNL.RM_MHL_FNL");
+		UnbakedDefinitionPaths = { TEXT("/RigMapper/Definitions/Raw/RM_MHL_MHH.RM_MHL_MHH"),
+			TEXT("/RigMapper/Definitions/Raw/RM_MHH_FNM.RM_MHH_FNM"),
+			TEXT("/RigMapper/Definitions/Raw/RM_FNM_FNH.RM_FNM_FNH"),
+			TEXT("/RigMapper/Definitions/Raw/RM_FNH_FNL.RM_FNH_FNL")
+		};
+		TestAnimSequencePath = TEXT("/Game/AutomationTestData/RigMapper/AnimSequences/AS_Invictus_MH3.AS_Invictus_MH3");
+	}
+	else if (InParameters == "MHL_CDL")
+	{
+		BakedDefinitionPath = TEXT("/RigMapper/Definitions/Baked/RM_MHL_CDL.RM_MHL_CDL");
+		UnbakedDefinitionPaths = { TEXT("/RigMapper/Definitions/Raw/RM_MHL_MHH.RM_MHL_MHH"), 
+			TEXT("/RigMapper/Definitions/Raw/RM_MHH_CDH.RM_MHH_CDH"),
+			TEXT("/RigMapper/Definitions/Raw/RM_CDH_CDL.RM_CDH_CDL")
+		};
+		TestAnimSequencePath = TEXT("/Game/AutomationTestData/RigMapper/AnimSequences/AS_Invictus_MH3.AS_Invictus_MH3");
+	}
+	else if (InParameters == "CDL_FNL")
+	{
+		BakedDefinitionPath = TEXT("/RigMapper/Definitions/Baked/RM_CDL_FNL.RM_CDL_FNL");
+		UnbakedDefinitionPaths = { TEXT("/RigMapper/Definitions/Raw/RM_CDL_CDH.RM_CDL_CDH"),
+			TEXT("/RigMapper/Definitions/Raw/RM_CDH_FNM.RM_CDH_FNM"),
+			TEXT("/RigMapper/Definitions/Raw/RM_FNM_FNH.RM_FNM_FNH"),
+			TEXT("/RigMapper/Definitions/Raw/RM_FNH_FNL.RM_FNH_FNL")
+		};
+		TestAnimSequencePath = TEXT("/Game/AutomationTestData/RigMapper/AnimSequences/FACIAL_3L_RIG_ROM.FACIAL_3L_RIG_ROM");
+	}
+
+	// load in the definitions and test sequence
+	TArray<TSharedPtr<FacialRigMapping::FRigMapper>> RigMappersUnbaked, RigMappersBaked;
+	TArray<URigMapperDefinition*> RigMapperDefinitionsUnbaked, RigMapperDefinitionsBaked;
+	RigMapperDefinitionsUnbaked.SetNum(UnbakedDefinitionPaths.Num());
+	RigMapperDefinitionsBaked.SetNum(1);
+	TArray<FRigMapperIndexMap> IndexMapsUnbaked, IndexMapsBaked;
+
+	RigMapperDefinitionsBaked[0] = LoadObject<URigMapperDefinition>((UObject*)GetTransientPackage(), *BakedDefinitionPath);
+	if (!RigMapperDefinitionsBaked[0])
+	{
+		UE_LOG(LogRigMapperTest, Error, TEXT("Failed to load asset from path: %s"), *BakedDefinitionPath);
+		return false;
+	}
+
+	for (int32 Def = 0; Def < RigMapperDefinitionsUnbaked.Num(); ++Def)
+	{
+		RigMapperDefinitionsUnbaked[Def] = LoadObject<URigMapperDefinition>((UObject*)GetTransientPackage(), *UnbakedDefinitionPaths[Def]);
+		if (!RigMapperDefinitionsUnbaked[Def])
+		{
+			UE_LOG(LogRigMapperTest, Error, TEXT("Failed to load asset from path: %s"), *UnbakedDefinitionPaths[Def]);
+			return false;
+		}
+	}
+
+	TestAnimSequence = LoadObject<UAnimSequence>((UObject*)GetTransientPackage(), *TestAnimSequencePath);
+
+	if (TestAnimSequence)
+	{
+		// run the comparison test
+		const FAnimationCurveData& CurveData = TestAnimSequence->GetDataModel()->GetCurveData();
+		TArray<TMap<FString, double>> CurveDataRigMapperInputs = GetTestSequenceCurves(TestAnimSequence, ControlSubStr);
+
+		bool bInitialized = InitializeRigMapping(RigMapperDefinitionsUnbaked, RigMappersUnbaked, IndexMapsUnbaked);
+		if (!bInitialized)
+		{
+			UE_LOG(LogRigMapperTest, Error, TEXT("Failed to initialize unbaked RigMapping"));
+			return false;
+		}
+
+		bInitialized = InitializeRigMapping(RigMapperDefinitionsBaked, RigMappersBaked, IndexMapsBaked);
+		if (!bInitialized)
+		{
+			UE_LOG(LogRigMapperTest, Error, TEXT("Failed to initialize baked RigMapping"));
+			return false;
+		}
+
+		if (CurveDataRigMapperInputs.IsEmpty())
+		{
+			UE_LOG(LogRigMapperTest, Error, TEXT("Expected test data to contain at least one frame"));
+			return false;
+		}
+
+		for (int32 Frame = 0; Frame < CurveDataRigMapperInputs.Num(); ++Frame)
+		{
+			// fill in any inputs missing in the anim sequence
+			for (int32 Input = 0; Input < RigMapperDefinitionsBaked[0]->Inputs.Num(); ++Input)
+			{
+				if (!CurveDataRigMapperInputs[Frame].Find(RigMapperDefinitionsBaked[0]->Inputs[Input]))
+				{
+					CurveDataRigMapperInputs[Frame].Add(RigMapperDefinitionsBaked[0]->Inputs[Input], 0);
+					if (Frame == 0)
+					{
+						UE_LOG(LogRigMapperTest, Warning, TEXT("Missing value in test anim sequence for control: %s , setting to 0"), *RigMapperDefinitionsBaked[0]->Inputs[Input]);
+					}
+				}
+			}
+		}
+
+		for (int32 Frame = 0; Frame < CurveDataRigMapperInputs.Num(); ++Frame)
+		{
+			// evaluate the rig mappers for both baked and unbaked cases; this evaluates the whole thing in the same way in which this is done in the anim node
+			TMap<FString, double> OutputsBaked, OutputsUnbaked;
+			EvaluateRigMapping(CurveDataRigMapperInputs[Frame], IndexMapsBaked, RigMappersBaked, OutputsBaked);
+			EvaluateRigMapping(CurveDataRigMapperInputs[Frame], IndexMapsUnbaked, RigMappersUnbaked, OutputsUnbaked);
+
+			// check that the round trip gives outputs which are identical to the inputs
+			if (!AreMapsEqualWithTolerance(OutputsUnbaked, OutputsBaked, CurveIdentifier, {}, Tolerance))
+			{
+				UE_LOG(LogRigMapperTest, Error, TEXT("Frame %d contains different values for unbaked vs. baked definition test"), Frame);
+				bResult = false;
+			}
+		}
+	}
+	else
+	{
+		// warn, but don't fail, as this test will only work in Beehive project
+		UE_LOG(LogRigMapperTest, Warning, TEXT("Failed to load test animation sequence from path: %s . Note that this test data is only available in project Sandbox/Anim/Beehive"), *TestAnimSequencePath);
+	}
+
+	return bResult;
+}
+
 #endif
\ No newline at end of file
