fixed a couple of undo bugs

diff --git a/Engine/Plugins/Experimental/Animation/RigMapper/Source/RigMapperEditor/Private/RigMapperDefinitionEditorToolkit.h b/Engine/Plugins/Experimental/Animation/RigMapper/Source/RigMapperEditor/Private/RigMapperDefinitionEditorToolkit.h
--- a/Engine/Plugins/Experimental/Animation/RigMapper/Source/RigMapperEditor/Private/RigMapperDefinitionEditorToolkit.h
+++ b/Engine/Plugins/Experimental/Animation/RigMapper/Source/RigMapperEditor/Private/RigMapperDefinitionEditorToolkit.h
@@ -6,6 +6,7 @@
 #include "RigMapperDefinition.h"
 
 #include "Toolkits/AssetEditorToolkit.h"
+#include "EditorUndoClient.h"
 
 enum class ERigMapperNodeType : uint8;
 class SRigMapperDefinitionStructureView;
@@ -14,7 +15,8 @@ class SRigMapperDefinitionGraphEditor;
 /**
  * The toolkit for the URigMapperDefinition asset editor
  */
-class RIGMAPPEREDITOR_API FRigMapperDefinitionEditorToolkit : public FAssetEditorToolkit
+class RIGMAPPEREDITOR_API FRigMapperDefinitionEditorToolkit : public FAssetEditorToolkit, 
+	public FSelfRegisteringEditorUndoClient
 {
 public:
 	void Initialize(URigMapperDefinition* InDefinition, EToolkitMode::Type InMode, TSharedPtr<IToolkitHost> InToolkitHost);
diff --git a/Engine/Plugins/Experimental/Animation/RigMapper/Source/RigMapperEditor/Private/RigMapperEditorModule.cpp b/Engine/Plugins/Experimental/Animation/RigMapper/Source/RigMapperEditor/Private/RigMapperEditorModule.cpp
--- a/Engine/Plugins/Experimental/Animation/RigMapper/Source/RigMapperEditor/Private/RigMapperEditorModule.cpp
+++ b/Engine/Plugins/Experimental/Animation/RigMapper/Source/RigMapperEditor/Private/RigMapperEditorModule.cpp
@@ -25,6 +25,7 @@
 #include "RigMapperGraph/RigMapperDefinitionEditorGraphNode.h"
 #include "Misc/ConfigCacheIni.h"
 #include "Misc/Paths.h"
+#include "ScopedTransaction.h"
 
 const FName FRigMapperEditorModule::AppIdentifier("RigMapperEditorApp");
 
@@ -281,8 +282,11 @@ void FRigMapperEditorModule::LoadFromJson(const FToolMenuContext& ToolMenuContex
 
 	if (!Definitions.IsEmpty() && GetUserPickedPath(JsonFile.FilePath))
 	{
+		FScopedTransaction Transaction(LOCTEXT("LoadFromJson", "Load RigMapper Definitions from Json file"));
+
 		for (URigMapperDefinition* Definition : Definitions)
 		{
+			Definition->Modify();
 			if (!Definition->LoadFromJsonFile(JsonFile))
 			{
 				FMessageLog MessageLog(MessageLogIdentifier);
@@ -292,6 +296,8 @@ void FRigMapperEditorModule::LoadFromJson(const FToolMenuContext& ToolMenuContex
 	
 				MessageLog.Error(FText::Format(LOCTEXT("RigMapperEditorActions_ImportFailed", "Failed to load definition \"{ObjectName}\" from Json. See output log for more details"), Arguments));
 				MessageLog.Open(EMessageSeverity::Error);
+				Transaction.Cancel();
+				break;
 			}
 			else
 			{
