fix weird selection behaviour in RigMapper Definitions details panel

diff --git a/Engine/Plugins/Experimental/Animation/RigMapper/Source/RigMapperEditor/Private/RigMapperDefinitionEditorToolkit.cpp b/Engine/Plugins/Experimental/Animation/RigMapper/Source/RigMapperEditor/Private/RigMapperDefinitionEditorToolkit.cpp
--- a/Engine/Plugins/Experimental/Animation/RigMapper/Source/RigMapperEditor/Private/RigMapperDefinitionEditorToolkit.cpp
+++ b/Engine/Plugins/Experimental/Animation/RigMapper/Source/RigMapperEditor/Private/RigMapperDefinitionEditorToolkit.cpp
@@ -258,13 +258,17 @@ void FRigMapperDefinitionEditorToolkit::HandleFinishedChangingProperties(const F
 	const FName ParentName = Event.GetMemberPropertyName();
 	const FName InputName = GET_MEMBER_NAME_CHECKED(URigMapperDefinition, Inputs);
 	const FName OutputName = GET_MEMBER_NAME_CHECKED(URigMapperDefinition, Outputs);
+	const FName FeaturesName = GET_MEMBER_NAME_CHECKED(URigMapperDefinition, Features);
 	
 	const bool bInput = PropertyName == InputName && ParentName == InputName;
 	const bool bOutput = !bInput && (PropertyName == OutputName || PropertyName == OutputName.ToString() + TEXT("_Key")) && ParentName == OutputName; // PropertyName will be Outputs_Key if renaming an output
 	const bool bName = !bInput && !bOutput && PropertyName == GET_MEMBER_NAME_CHECKED(FRigMapperFeature, Name);
+	const bool bFeaturesChildProperty = ParentName == FeaturesName;
 
-	const bool bSelectNew = ChangeType == EPropertyChangeType::Duplicate || ChangeType == EPropertyChangeType::ArrayAdd;
-	const bool bRename = ChangeType == EPropertyChangeType::ValueSet && (bInput || bOutput || bName);
+	const bool bSelectNew = (ChangeType == EPropertyChangeType::Duplicate || ChangeType == EPropertyChangeType::ArrayAdd);
+
+
+	const bool bRename = ChangeType == EPropertyChangeType::ValueSet && (bInput || bOutput || bName || bFeaturesChildProperty);
 	
 	GraphEditor->RebuildGraph();
 	StructureView->RebuildTree();
@@ -285,7 +289,7 @@ void FRigMapperDefinitionEditorToolkit::HandleFinishedChangingProperties(const F
 			StructureView->SelectNode(OutputsInOrder[OutputIndex], ERigMapperNodeType::Output, true);
 			GraphEditor->SelectNodes({}, {}, { OutputsInOrder[OutputIndex] });
 		}
-		else if (ParentName == GET_MEMBER_NAME_CHECKED(URigMapperDefinition, Features))
+		else if (bFeaturesChildProperty)
 		{
 			if (const int32 MultiplyIndex = Event.GetArrayIndex(GET_MEMBER_NAME_CHECKED(FRigMapperFeatureDefinitions, Multiply).ToString()); MultiplyIndex != INDEX_NONE)
 			{
