removed used of LogTemp and gave RigMapper logs a better name. Also log warnings if a definition is used at runtime which is not validated (for example if it is edited)

diff --git a/Engine/Plugins/Experimental/Animation/RigMapper/Source/RigMapper/Private/AnimNode_RigMapper.cpp b/Engine/Plugins/Experimental/Animation/RigMapper/Source/RigMapper/Private/AnimNode_RigMapper.cpp
--- a/Engine/Plugins/Experimental/Animation/RigMapper/Source/RigMapper/Private/AnimNode_RigMapper.cpp
+++ b/Engine/Plugins/Experimental/Animation/RigMapper/Source/RigMapper/Private/AnimNode_RigMapper.cpp
@@ -7,6 +7,7 @@
 #include "Animation/AnimCurveUtils.h"
 #include "Engine/SkeletalMesh.h"
 #include "Animation/AnimInstanceProxy.h"
+#include "RigMapperLog.h"
 
 DEFINE_STAT(STAT_RigMapperEvaluateNode);
 DEFINE_STAT(STAT_RigMapperEvaluateRig);
@@ -80,7 +81,7 @@ void FAnimNode_RigMapper::PreUpdate(const UAnimInstance* InAnimInstance)
 	{
 		for (int32 DefinitionIndex = 0; DefinitionIndex < DefinitionsToCheck->Num(); DefinitionIndex++)
 		{
-			if ((*DefinitionsToCheck)[DefinitionIndex] != LoadedDefinitions[DefinitionIndex])
+			if ((*DefinitionsToCheck)[DefinitionIndex] != LoadedDefinitions[DefinitionIndex] || !(*DefinitionsToCheck)[DefinitionIndex]->IsDefinitionValidated()) // check if individual definition has been changed or editted
 			{
 				bReInit = true;
 				break;
@@ -238,12 +239,14 @@ bool FAnimNode_RigMapper::InitializeRigMapping(USkeletalMesh* TargetMesh)
 	TArray<URigMapperDefinition*> DefinitionsToLoad;
 
 	// Retrieve the definitions to load and use (either the ones set on the node, or the ones overriden by the SKM asset user data)
+	bool bSkelMeshOverride = false;
 	if (TargetMesh)
 	{
 		LoadedUserData = Cast<URigMapperDefinitionUserData>(TargetMesh->GetAssetUserDataOfClass(URigMapperDefinitionUserData::StaticClass()));
 		if (LoadedUserData)
 		{
 			DefinitionsToLoad = LoadedUserData->Definitions;
+			bSkelMeshOverride = true;
 		}
 	}
 	if (DefinitionsToLoad.IsEmpty())
@@ -253,15 +256,41 @@ bool FAnimNode_RigMapper::InitializeRigMapping(USkeletalMesh* TargetMesh)
 
 	if (DefinitionsToLoad.IsEmpty())
 	{
+		if (bSkelMeshOverride)
+		{
+			UE_LOG(LogRigMapper, Error, TEXT("Target SkelMesh %s has Rig Mapper Definition User Data set, but this contains no definitions."), *TargetMesh->GetName())
+		}
 		return false;
 	}
 	for (int i = 0; i < DefinitionsToLoad.Num(); i++)
 	{
 		if (!DefinitionsToLoad[i])
 		{
-			UE_LOG(LogTemp, Error, TEXT("Invalid definition at index %d"), i)
+			if (bSkelMeshOverride)
+			{
+				UE_LOG(LogRigMapper, Error, TEXT("Target SkelMesh %s has Rig Mapper Definition User Data set which contains an empty RigMapper definition asset at index %d; please check that all User Data RigMapper definitions have been set to valid assets."), *TargetMesh->GetName(), i)
+			}
+			else
+			{
+				UE_LOG(LogRigMapper, Error, TEXT("Rig Mapper anim node contains an empty RigMapper definition asset at index %d; please check that all RigMapper node definitions have been set to valid assets."), i)
+			}
 			return false;
 		}
+		else
+		{
+			// check if the definition is validated ... if not, warn
+			if (!DefinitionsToLoad[i]->IsDefinitionValidated())
+			{
+				if (bSkelMeshOverride)
+				{
+					UE_LOG(LogRigMapper, Warning, TEXT("Target SkelMesh %s has Rig Mapper Definition User Data set which contains an unvalidated RigMapper definition asset at index %d; please ensure that all User Data RigMapper definitions have been set to validated assets."), *TargetMesh->GetName(), i)
+				}
+				else
+				{
+					UE_LOG(LogRigMapper, Warning, TEXT("Rig Mapper anim node contains an unvalidated RigMapper definition asset at index %d; please ensure that all RigMapper node definitions have been set to validated assets."), i)
+				}
+			}
+		}
 	}
 
 	// Initialize the actual RigMappers from the definitions to load
@@ -274,7 +303,7 @@ bool FAnimNode_RigMapper::InitializeRigMapping(USkeletalMesh* TargetMesh)
 	
 		if (!RigMapper->LoadDefinition(Definition))
 		{
-			UE_LOG(LogTemp, Error, TEXT("Could not load definition %s"), *Definition->GetPathName())
+			UE_LOG(LogRigMapper, Error, TEXT("Could not load definition %s"), *Definition->GetPathName())
 			RigMappers.Empty();
 			return false;
 		}
diff --git a/Engine/Plugins/Experimental/Animation/RigMapper/Source/RigMapper/Private/RigMapper.cpp b/Engine/Plugins/Experimental/Animation/RigMapper/Source/RigMapper/Private/RigMapper.cpp
--- a/Engine/Plugins/Experimental/Animation/RigMapper/Source/RigMapper/Private/RigMapper.cpp
+++ b/Engine/Plugins/Experimental/Animation/RigMapper/Source/RigMapper/Private/RigMapper.cpp
@@ -4,6 +4,7 @@
 
 #include "RigMapperDefinition.h"
 #include "UObject/SoftObjectPath.h"
+#include "RigMapperLog.h"
 
 namespace FacialRigMapping
 {
@@ -270,7 +271,7 @@ namespace FacialRigMapping
 
 		if (!Evaluate_Static(InputValue, Keys, OutValue))
 		{
-			UE_LOG(LogTemp, Warning, TEXT("PiecewiseLinear could not calculate the output value"))
+			UE_LOG(LogRigMapper, Warning, TEXT("PiecewiseLinear could not calculate the output value"))
 			return false;	
 		}
 
diff --git a/Engine/Plugins/Experimental/Animation/RigMapper/Source/RigMapper/Private/RigMapperDefinition.cpp b/Engine/Plugins/Experimental/Animation/RigMapper/Source/RigMapper/Private/RigMapperDefinition.cpp
--- a/Engine/Plugins/Experimental/Animation/RigMapper/Source/RigMapper/Private/RigMapperDefinition.cpp
+++ b/Engine/Plugins/Experimental/Animation/RigMapper/Source/RigMapper/Private/RigMapperDefinition.cpp
@@ -10,8 +10,7 @@
 #include "Serialization/JsonReader.h"
 #include "Serialization/JsonWriter.h"
 #include "Serialization/JsonSerializer.h"
-
-DEFINE_LOG_CATEGORY(LogRigMapperDefinition)
+#include "RigMapperLog.h"
 
 bool FRigMapperFeature::IsValid(const TArray<FString>& InputNames, bool bWarn) const
 {
@@ -19,7 +18,7 @@ bool FRigMapperFeature::IsValid(const TArray<FString>& InputNames, bool bWarn) c
 	
 	if (Name.IsEmpty())
 	{
-		UE_LOG(LogRigMapperDefinition, Warning, TEXT("Invalid (empty) feature name"))
+		UE_LOG(LogRigMapper, Warning, TEXT("Invalid (empty) feature name"))
 		bValid = false;
 	}
 
@@ -28,13 +27,13 @@ bool FRigMapperFeature::IsValid(const TArray<FString>& InputNames, bool bWarn) c
 	
 	if (FeatureInputs.Contains(Name))
 	{
-		UE_LOG(LogRigMapperDefinition, Warning, TEXT("Feature %s is referencing itself"), *Name)
+		UE_LOG(LogRigMapper, Warning, TEXT("Feature %s is referencing itself"), *Name)
 		bValid = false;
 	}
 
 	if (FeatureInputs.IsEmpty())
 	{
-		UE_LOG(LogRigMapperDefinition, Warning, TEXT("Feature %s does not reference any input"), *Name)
+		UE_LOG(LogRigMapper, Warning, TEXT("Feature %s does not reference any input"), *Name)
 		bValid = false;
 	}
 	
@@ -42,7 +41,7 @@ bool FRigMapperFeature::IsValid(const TArray<FString>& InputNames, bool bWarn) c
 	{
 		if (!InputNames.Contains(Input))
 		{
-			UE_LOG(LogRigMapperDefinition, Warning, TEXT("Undefined input or feature %s referenced in feature %s"), *Input, *Name)
+			UE_LOG(LogRigMapper, Warning, TEXT("Undefined input or feature %s referenced in feature %s"), *Input, *Name)
 			bValid = false;
 		}
 	}
@@ -56,14 +55,14 @@ bool FRigMapperFeature::GetJsonArray(TSharedPtr<FJsonObject> JsonObject, TArray<
 	{
 		if (!JsonObject->HasTypedField<EJson::Object>(*OwnerIdentifier))
 		{
-			UE_LOG(LogRigMapperDefinition, Warning, TEXT("Missing '%s' field for feature %s"), *OwnerIdentifier, *Name)
+			UE_LOG(LogRigMapper, Warning, TEXT("Missing '%s' field for feature %s"), *OwnerIdentifier, *Name)
 			return false;
 		}
 		JsonObject = JsonObject->GetObjectField(OwnerIdentifier);
 	}
 	if (!JsonObject->HasTypedField<EJson::Array>(*Identifier))
 	{
-		UE_LOG(LogRigMapperDefinition, Warning, TEXT("Missing '%s' field for feature %s"), *Identifier, *Name)
+		UE_LOG(LogRigMapper, Warning, TEXT("Missing '%s' field for feature %s"), *Identifier, *Name)
 		return false;
 	}
 	OutArray = JsonObject->GetArrayField(*Identifier);
@@ -80,7 +79,7 @@ bool FRigMapperMultiplyFeature::LoadFromJsonObject(const TSharedPtr<FJsonObject>
 	}
 	if (InputFeatures.Num() < 2)
 	{
-		UE_LOG(LogRigMapperDefinition, Warning, TEXT("Feature %s does not reference enough input"), *Name)
+		UE_LOG(LogRigMapper, Warning, TEXT("Feature %s does not reference enough input"), *Name)
 		return false;
 	}
 	for (const auto& InputFeature : InputFeatures)
@@ -97,7 +96,7 @@ bool FRigMapperMultiplyFeature::IsValid(const TArray<FString>& InputNames, bool
 	
 	if (bValid && Inputs.Num() < 2)
 	{
-		UE_LOG(LogRigMapperDefinition, Warning, TEXT("Feature %s does not reference enough input"), *Name)
+		UE_LOG(LogRigMapper, Warning, TEXT("Feature %s does not reference enough input"), *Name)
 		return false;
 	}
 	
@@ -119,7 +118,7 @@ bool FRigMapperWsFeature::LoadFromJsonObject(const TSharedPtr<FJsonObject>& Json
 	}
 	if (JsonWeights.Num() != JsonInputs.Num())
 	{
-		UE_LOG(LogRigMapperDefinition, Warning, TEXT("Number of inputs does not match number of weights for feature %s"), *Name)
+		UE_LOG(LogRigMapper, Warning, TEXT("Number of inputs does not match number of weights for feature %s"), *Name)
 		return false;
 	}
 	for (int32 InputIndex = 0; InputIndex < JsonInputs.Num(); InputIndex++)
@@ -139,7 +138,7 @@ bool FRigMapperWsFeature::LoadFromJsonObject(const TSharedPtr<FJsonObject>& Json
 	}
 	if (Range.bHasLowerBound && Range.bHasUpperBound && Range.LowerBound > Range.UpperBound)
 	{
-		UE_LOG(LogRigMapperDefinition, Warning, TEXT("Invalid range for feature %s"), *Name)
+		UE_LOG(LogRigMapper, Warning, TEXT("Invalid range for feature %s"), *Name)
 		return false;
 	}
 	return true;
@@ -156,15 +155,15 @@ bool FRigMapperWsFeature::IsValid(const TArray<FString>& InputNames, bool bWarn)
 	}
 	if (bWarn && TotalWeight == 0)
 	{
-		UE_LOG(LogRigMapperDefinition, Warning, TEXT("Total Weights for feature %s add up to 0"), *Name)
+		UE_LOG(LogRigMapper, Warning, TEXT("Total Weights for feature %s add up to 0"), *Name)
 	}
 	if (bWarn && (TotalWeight < -1 || TotalWeight > 1) && (!Range.bHasLowerBound || !Range.bHasUpperBound))
 	{
-		UE_LOG(LogRigMapperDefinition, Warning, TEXT("Total Weights for feature %s are quite high (%f) even though a range was not set"), *Name, TotalWeight)
+		UE_LOG(LogRigMapper, Warning, TEXT("Total Weights for feature %s are quite high (%f) even though a range was not set"), *Name, TotalWeight)
 	}
 	if (Range.bHasLowerBound && Range.bHasUpperBound && Range.LowerBound > Range.UpperBound)
 	{
-		UE_LOG(LogRigMapperDefinition, Warning, TEXT("Range of [%f-%f] for feature %s is invalid"), Range.LowerBound, Range.UpperBound, *Name)
+		UE_LOG(LogRigMapper, Warning, TEXT("Range of [%f-%f] for feature %s is invalid"), Range.LowerBound, Range.UpperBound, *Name)
 		bValid = false;
 	}
 	
@@ -185,7 +184,7 @@ bool FRigMapperSdkFeature::LoadFromJsonObject(const TSharedPtr<FJsonObject>& Jso
 	}
 	if (JsonInputs.Num() != 1)
 	{
-		UE_LOG(LogRigMapperDefinition, Warning, TEXT("Sdk feature %s should have a single element in the 'input_features' array field"), *Name)
+		UE_LOG(LogRigMapper, Warning, TEXT("Sdk feature %s should have a single element in the 'input_features' array field"), *Name)
 		return false;
 	}
 	Input = JsonInputs[0]->AsString();
@@ -204,11 +203,11 @@ bool FRigMapperSdkFeature::LoadFromJsonObject(const TSharedPtr<FJsonObject>& Jso
 
 	if (InValues.Num() < 2)
 	{
-		UE_LOG(LogRigMapperDefinition, Warning, TEXT("Not enough keys for SDK feature %s (expected minimum 2, got %d)"), *Name, InValues.Num())
+		UE_LOG(LogRigMapper, Warning, TEXT("Not enough keys for SDK feature %s (expected minimum 2, got %d)"), *Name, InValues.Num())
 	}
 	if (InValues.Num() != OutValues.Num())
 	{
-		UE_LOG(LogRigMapperDefinition, Warning, TEXT("Number of input values does not match number of output values for feature %s"), *Name)
+		UE_LOG(LogRigMapper, Warning, TEXT("Number of input values does not match number of output values for feature %s"), *Name)
 		return false;
 	}
 	for (int32 ValueIndex = 0 ; ValueIndex < InValues.Num(); ValueIndex++)
@@ -226,7 +225,7 @@ bool FRigMapperSdkFeature::IsValid(const TArray<FString>& InputNames, bool bWarn
 
 	if (Keys.Num() < 2)
 	{
-		UE_LOG(LogRigMapperDefinition, Warning, TEXT("Not enough keys for SDK feature %s (expected minimum 2, got %d)"), *Name, Keys.Num())
+		UE_LOG(LogRigMapper, Warning, TEXT("Not enough keys for SDK feature %s (expected minimum 2, got %d)"), *Name, Keys.Num())
 		bValid = false;
 	}
 	
@@ -237,7 +236,7 @@ bool FRigMapperFeatureDefinitions::AddFromJsonObject(const FString& FeatureName,
 {
 	if (!JsonObject->HasTypedField<EJson::String>(TEXT("type")))
 	{
-		UE_LOG(LogRigMapperDefinition, Warning, TEXT("Missing 'type' field for feature %s"), *FeatureName)
+		UE_LOG(LogRigMapper, Warning, TEXT("Missing 'type' field for feature %s"), *FeatureName)
 		return false;
 	}
 	const FString FeatureType = JsonObject->GetStringField(TEXT("type"));
@@ -261,7 +260,7 @@ bool FRigMapperFeatureDefinitions::AddFromJsonObject(const FString& FeatureName,
 	}
 	else
 	{
-		UE_LOG(LogRigMapperDefinition, Warning, TEXT("Invalid type for feature %s (%s)"), *FeatureName, *FeatureType)
+		UE_LOG(LogRigMapper, Warning, TEXT("Invalid type for feature %s (%s)"), *FeatureName, *FeatureType)
 		return bValidFeature;
 	}
 	return bValidFeature;
@@ -301,14 +300,14 @@ bool FRigMapperFeatureDefinitions::IsValid(const TArray<FString>& InputNames, bo
 	{
 		if (InputNames.Contains(FeatureName))
 		{
-			UE_LOG(LogRigMapperDefinition, Warning, TEXT("Conflicting input and feature name: %s"), *FeatureName)
+			UE_LOG(LogRigMapper, Warning, TEXT("Conflicting input and feature name: %s"), *FeatureName)
 			bValid = false;
 		}
 		if (bDuplicateFeatureName)
 		{
 			if (CheckFeatureNames.Contains(FeatureName))
 			{
-				UE_LOG(LogRigMapperDefinition, Warning, TEXT("Duplicate feature name: %s"), *FeatureName)
+				UE_LOG(LogRigMapper, Warning, TEXT("Duplicate feature name: %s"), *FeatureName)
 				bValid = false;
 			}
 			else
@@ -373,7 +372,7 @@ bool URigMapperDefinition::LoadFromJsonFile(const FFilePath& JsonFilePath)
 		
 	if (!FFileHelper::LoadFileToString(JsonAsString, *JsonFilePath.FilePath))
 	{
-		UE_LOG(LogRigMapperDefinition, Warning, TEXT("Could not open json file"))
+		UE_LOG(LogRigMapper, Warning, TEXT("Could not open json file"))
 		return false;
 	}
 
@@ -397,7 +396,7 @@ bool URigMapperDefinition::LoadFromJsonString(const FString& JsonString)
 	
 	if (!FJsonSerializer::Deserialize(JsonReader, JsonObject) && JsonObject.IsValid())
 	{
-		UE_LOG(LogRigMapperDefinition, Warning, TEXT("Could not deserialize json data to load definition"))
+		UE_LOG(LogRigMapper, Warning, TEXT("Could not deserialize json data to load definition"))
 		return false;
 	}
 
@@ -436,7 +435,7 @@ bool URigMapperDefinition::LoadInputsFromJsonObject(const TSharedPtr<FJsonObject
 {
 	if (!JsonObject->HasTypedField(TEXT("inputs"), EJson::Array))
 	{
-		UE_LOG(LogRigMapperDefinition, Warning, TEXT("Missing inputs field"))
+		UE_LOG(LogRigMapper, Warning, TEXT("Missing inputs field"))
 		bDefinitionValidated = false;
 		return false;
 	}
@@ -445,7 +444,7 @@ bool URigMapperDefinition::LoadInputsFromJsonObject(const TSharedPtr<FJsonObject
 		const FString InputString = InputAttr->AsString();
 		if (Inputs.Contains(InputString))
 		{
-			UE_LOG(LogRigMapperDefinition, Warning, TEXT("Duplicate input was found: %s"), *InputString)
+			UE_LOG(LogRigMapper, Warning, TEXT("Duplicate input was found: %s"), *InputString)
 		}
 		else
 		{
@@ -454,7 +453,7 @@ bool URigMapperDefinition::LoadInputsFromJsonObject(const TSharedPtr<FJsonObject
 	}
 	if (Inputs.IsEmpty())
 	{
-		UE_LOG(LogRigMapperDefinition, Warning, TEXT("Not enough inputs"))
+		UE_LOG(LogRigMapper, Warning, TEXT("Not enough inputs"))
 		bDefinitionValidated = false;
 	}
 	return bDefinitionValidated;
@@ -469,13 +468,13 @@ bool URigMapperDefinition::LoadNullOutputsFromJsonObject(const TSharedPtr<FJsonO
 			const FString NullOutputString = NullOutputAttr->AsString();
 			if (NullOutputs.Contains(NullOutputString))
 			{
-				UE_LOG(LogRigMapperDefinition, Warning, TEXT("Duplicate null output was found: %s"), *NullOutputString)
+				UE_LOG(LogRigMapper, Warning, TEXT("Duplicate null output was found: %s"), *NullOutputString)
 			}
 			else
 			{
 				if (Outputs.Contains(NullOutputString))
 				{
-					UE_LOG(LogRigMapperDefinition, Warning, TEXT("Null output conflicts with existing output: %s"), *NullOutputString)
+					UE_LOG(LogRigMapper, Warning, TEXT("Null output conflicts with existing output: %s"), *NullOutputString)
 					bDefinitionValidated = false;
 				}
 				else
@@ -492,7 +491,7 @@ bool URigMapperDefinition::LoadFeaturesFromJsonObject(const TSharedPtr<FJsonObje
 {
 	if (!JsonObject->HasTypedField(TEXT("features"), EJson::Object))
 	{
-		UE_LOG(LogRigMapperDefinition, Warning, TEXT("Missing features field"))
+		UE_LOG(LogRigMapper, Warning, TEXT("Missing features field"))
 		bDefinitionValidated = false;
 		return false;
 	}
@@ -500,7 +499,7 @@ bool URigMapperDefinition::LoadFeaturesFromJsonObject(const TSharedPtr<FJsonObje
 	{
 		if (Inputs.Contains(FeatureInfo.Key))
 		{
-			UE_LOG(LogRigMapperDefinition, Warning, TEXT("Feature conflicting with input of similar name: %s"), *FeatureInfo.Key)
+			UE_LOG(LogRigMapper, Warning, TEXT("Feature conflicting with input of similar name: %s"), *FeatureInfo.Key)
 			bDefinitionValidated = false;
 			continue;
 		}
@@ -520,7 +519,7 @@ bool URigMapperDefinition::LoadOutputsFromJsonObject(const TSharedPtr<FJsonObjec
 {
 	if (!JsonObject->HasTypedField(TEXT("outputs"), EJson::Object))
 	{
-		UE_LOG(LogRigMapperDefinition, Warning, TEXT("Missing outputs field"))
+		UE_LOG(LogRigMapper, Warning, TEXT("Missing outputs field"))
 		bDefinitionValidated = false;
 		return false;
 	}
@@ -528,13 +527,13 @@ bool URigMapperDefinition::LoadOutputsFromJsonObject(const TSharedPtr<FJsonObjec
 	{
 		if (OutputAttr.Key.IsEmpty())
 		{
-			UE_LOG(LogRigMapperDefinition, Warning, TEXT("Invalid output with empty name"))
+			UE_LOG(LogRigMapper, Warning, TEXT("Invalid output with empty name"))
 			continue;
 		}
 		FString OutputValue;
 		if (!OutputAttr.Value.IsValid() || OutputAttr.Value->IsNull() || OutputAttr.Value->AsString().IsEmpty())
 		{
-			UE_LOG(LogRigMapperDefinition, Warning, TEXT("Invalid value for output %s"), *OutputAttr.Key)
+			UE_LOG(LogRigMapper, Warning, TEXT("Invalid value for output %s"), *OutputAttr.Key)
 			bDefinitionValidated = false;
 		}
 		else
@@ -542,7 +541,7 @@ bool URigMapperDefinition::LoadOutputsFromJsonObject(const TSharedPtr<FJsonObjec
 			OutputValue = OutputAttr.Value->AsString();
 			if (!Inputs.Contains(OutputValue) && !FeatureNames.Contains(OutputValue))
 			{
-				UE_LOG(LogRigMapperDefinition, Warning, TEXT("Could not find corresponding input/feature for output %s (%s)"), *OutputAttr.Key, *OutputValue)
+				UE_LOG(LogRigMapper, Warning, TEXT("Could not find corresponding input/feature for output %s (%s)"), *OutputAttr.Key, *OutputValue)
 				bDefinitionValidated = false;
 			}
 		}
@@ -550,7 +549,7 @@ bool URigMapperDefinition::LoadOutputsFromJsonObject(const TSharedPtr<FJsonObjec
 	}
 	if (Outputs.IsEmpty())
 	{
-		UE_LOG(LogRigMapperDefinition, Warning, TEXT("Not enough outputs"))
+		UE_LOG(LogRigMapper, Warning, TEXT("Not enough outputs"))
 		bDefinitionValidated = false;
 	}
 	return bDefinitionValidated;
@@ -662,6 +661,11 @@ bool URigMapperDefinition::ExportAsJsonFile(const FFilePath& JsonFilePath) const
 	return false;
 }
 
+bool URigMapperDefinition::IsDefinitionValidated() const
+{
+	return bDefinitionValidated;
+}
+
 bool URigMapperDefinition::IsDefinitionValid(bool bWarn) const
 {
 	if (bDefinitionValidated)
@@ -673,12 +677,12 @@ bool URigMapperDefinition::IsDefinitionValid(bool bWarn) const
 
 	if (Inputs.IsEmpty())
 	{
-		UE_LOG(LogRigMapperDefinition, Warning, TEXT("Not enough inputs"))
+		UE_LOG(LogRigMapper, Warning, TEXT("Not enough inputs"))
 		bValid = false;
 	}
 	if (Outputs.IsEmpty())
 	{
-		UE_LOG(LogRigMapperDefinition, Warning, TEXT("Not enough outputs"))
+		UE_LOG(LogRigMapper, Warning, TEXT("Not enough outputs"))
 		bValid = false;
 	}
 
@@ -689,7 +693,7 @@ bool URigMapperDefinition::IsDefinitionValid(bool bWarn) const
 		// todo: check with skeleton and or control rig ?
 		if (CheckInputs.Contains(Input))
 		{
-			UE_LOG(LogRigMapperDefinition, Warning, TEXT("Duplicate input %s"), *Input)
+			UE_LOG(LogRigMapper, Warning, TEXT("Duplicate input %s"), *Input)
 			bValid = false;
 		}
 		else
@@ -704,7 +708,7 @@ bool URigMapperDefinition::IsDefinitionValid(bool bWarn) const
 		// todo: check with skeleton and or control rig ?
 		if (!Inputs.Contains(Output.Value) && !FeatureNames.Contains(Output.Value))
 		{
-			UE_LOG(LogRigMapperDefinition, Warning, TEXT("Output %s does not link to any existing input or feature"), *Output.Key)
+			UE_LOG(LogRigMapper, Warning, TEXT("Output %s does not link to any existing input or feature"), *Output.Key)
 			bValid = false;
 		}
 	}
@@ -719,12 +723,12 @@ bool URigMapperDefinition::IsDefinitionValid(bool bWarn) const
 	{
 		if (Outputs.Contains(NullOutput))
 		{
-			UE_LOG(LogRigMapperDefinition, Warning, TEXT("Output is also defined as a NullOutput %s"), *NullOutput)
+			UE_LOG(LogRigMapper, Warning, TEXT("Output is also defined as a NullOutput %s"), *NullOutput)
 			bValid = false;
 		}
 		if (CheckNullOutputs.Contains(NullOutput))
 		{
-			UE_LOG(LogRigMapperDefinition, Warning, TEXT("Duplicate NullOutput %s"), *NullOutput)
+			UE_LOG(LogRigMapper, Warning, TEXT("Duplicate NullOutput %s"), *NullOutput)
 			bValid = false;
 		}
 		else
@@ -738,12 +742,12 @@ bool URigMapperDefinition::IsDefinitionValid(bool bWarn) const
 
 bool URigMapperDefinition::Validate()
 {
-	UE_LOG(LogRigMapperDefinition, Log, TEXT("Validating definition %s"), *GetName())
+	UE_LOG(LogRigMapper, Log, TEXT("Validating definition %s"), *GetName())
 	
 	bDefinitionValidated = false;
 	bDefinitionValidated = IsDefinitionValid(true);
 
-	UE_LOG(LogRigMapperDefinition, Log, TEXT("Definition %s is %s"), *GetName(), bDefinitionValidated ? TEXT("valid") : TEXT("invalid"))
+	UE_LOG(LogRigMapper, Log, TEXT("Definition %s is %s"), *GetName(), bDefinitionValidated ? TEXT("valid") : TEXT("invalid"))
 	
 	return bDefinitionValidated;
 }
@@ -770,7 +774,7 @@ bool URigMapperLinkedDefinitions::GetBakedInputRec(const FString& InputName, con
 			{
 				if (!SourceDefinitions[DefinitionIndex - 1]->NullOutputs.Contains(InputName))
 				{
-					UE_LOG(LogRigMapperDefinition, Warning, TEXT("Input %s on definition %s (layer %d) does not match any output from definition %s (layer %d)"), *InputName, *SourceDefinitions[DefinitionIndex]->GetName(), DefinitionIndex, *SourceDefinitions[DefinitionIndex - 1]->GetName(), DefinitionIndex - 1)
+					UE_LOG(LogRigMapper, Warning, TEXT("Input %s on definition %s (layer %d) does not match any output from definition %s (layer %d)"), *InputName, *SourceDefinitions[DefinitionIndex]->GetName(), DefinitionIndex, *SourceDefinitions[DefinitionIndex - 1]->GetName(), DefinitionIndex - 1)
 					bOutMissingIsNotNullOutput = true;
 				}
 				return false;
@@ -788,7 +792,7 @@ bool URigMapperLinkedDefinitions::GetBakedInputRec(const FString& InputName, con
 	{
 		return Feature->BakeInput(this, DefinitionIndex, OutBakedInput);
 	}
-	UE_LOG(LogRigMapperDefinition, Warning, TEXT("Could not bake input %s"), *InputName)
+	UE_LOG(LogRigMapper, Warning, TEXT("Could not bake input %s"), *InputName)
 	return false;
 }
 
@@ -983,7 +987,7 @@ bool URigMapperLinkedDefinitions::AddBakedInputFeature(const TSharedPtr<FRigMapp
 	{
 		if (!SourceDefinitions[0]->Inputs.Contains(Feature->Name))
 		{
-			UE_LOG(LogRigMapperDefinition, Warning, TEXT("Baked input could not be found within the lower level inputs"))
+			UE_LOG(LogRigMapper, Warning, TEXT("Baked input could not be found within the lower level inputs"))
 			return false;
 		}
 		BakedDefinition->Inputs.AddUnique(Feature->Name);
@@ -1040,7 +1044,7 @@ TArray<FRigMapperFeature::FBakedInput> URigMapperLinkedDefinitions::GetBakedInpu
 		}
 		else
 		{
-			UE_LOG(LogRigMapperDefinition, Warning, TEXT("Could not bake input %s associated to output %s"), *Out.Value, *Out.Key)
+			UE_LOG(LogRigMapper, Warning, TEXT("Could not bake input %s associated to output %s"), *Out.Value, *Out.Key)
 		}
 	}
 
@@ -1051,23 +1055,23 @@ TArray<FRigMapperFeature::FBakedInput> URigMapperLinkedDefinitions::GetBakedInpu
 
 bool URigMapperLinkedDefinitions::BakeDefinitions()
 {
-	UE_LOG(LogRigMapperDefinition, Log, TEXT("Baking linked definition %s"), *GetName())
+	UE_LOG(LogRigMapper, Log, TEXT("Baking linked definition %s"), *GetName())
 	
 	if (!BakedDefinition)
 	{
-		UE_LOG(LogRigMapperDefinition, Warning, TEXT("Baked definition is unset"))
+		UE_LOG(LogRigMapper, Warning, TEXT("Baked definition is unset"))
 		return false;
 	}
 	if (SourceDefinitions.Num() < 2 || !SourceDefinitions[0] || !SourceDefinitions.Last())
 	{
-		UE_LOG(LogRigMapperDefinition, Warning, TEXT("Baking requires a minimum of 2 source definitions"))
+		UE_LOG(LogRigMapper, Warning, TEXT("Baking requires a minimum of 2 source definitions"))
 		return false;
 	}
 	for (int32 DefIndex = 0; DefIndex < SourceDefinitions.Num(); DefIndex++)
 	{
 		if (!SourceDefinitions[DefIndex]|| !SourceDefinitions[DefIndex]->IsDefinitionValid(true))
 		{
-			UE_LOG(LogRigMapperDefinition, Warning, TEXT("Invalid source definition at index %d. Make sure to revalidate the asset if necessary"), DefIndex)
+			UE_LOG(LogRigMapper, Warning, TEXT("Invalid source definition at index %d. Make sure to revalidate the asset if necessary"), DefIndex)
 			return false;
 		}
 	}
@@ -1076,12 +1080,12 @@ bool URigMapperLinkedDefinitions::BakeDefinitions()
 	const TArray<FRigMapperFeature::FBakedInput> BakedInputs = GetBakedInputs(PairedOutputs);
 	if (BakedInputs.Num() != PairedOutputs.Num())
 	{
-		UE_LOG(LogRigMapperDefinition, Warning, TEXT("The number of baked inputs do not match the expected number of outputs"))
+		UE_LOG(LogRigMapper, Warning, TEXT("The number of baked inputs do not match the expected number of outputs"))
 		return false;
 	}
 	AddBakedInputs(BakedInputs, PairedOutputs);
 
-	UE_LOG(LogRigMapperDefinition, Log, TEXT("Finished baking linked definition %s"), *GetName())
+	UE_LOG(LogRigMapper, Log, TEXT("Finished baking linked definition %s"), *GetName())
 	
 	return BakedDefinition->Validate();
 }
@@ -1110,20 +1114,20 @@ bool URigMapperLinkedDefinitions::AreLinkedDefinitionsValid() const
 	
 	if (!bOk)
 	{
-		UE_LOG(LogRigMapperDefinition, Warning, TEXT("Failed to validate the baked definition"))
+		UE_LOG(LogRigMapper, Warning, TEXT("Failed to validate the baked definition"))
 	}
 
 	for (int32 DefinitionIndex = 0; DefinitionIndex < SourceDefinitions.Num(); DefinitionIndex++)
 	{
 		if (!SourceDefinitions[DefinitionIndex])
 		{
-			UE_LOG(LogRigMapperDefinition, Warning, TEXT("Source definition %d is unset"), DefinitionIndex)
+			UE_LOG(LogRigMapper, Warning, TEXT("Source definition %d is unset"), DefinitionIndex)
 			bOk = false;
 			continue;
 		}
 		if (!SourceDefinitions[DefinitionIndex]->IsDefinitionValid())
 		{
-			UE_LOG(LogRigMapperDefinition, Warning, TEXT("Source definition %d is invalid"), DefinitionIndex)
+			UE_LOG(LogRigMapper, Warning, TEXT("Source definition %d is invalid"), DefinitionIndex)
 			bOk = false;
 		}
 		if (DefinitionIndex > 0)
@@ -1132,7 +1136,7 @@ bool URigMapperLinkedDefinitions::AreLinkedDefinitionsValid() const
 			{
 				if (!SourceDefinitions[DefinitionIndex - 1]->Outputs.Contains(Input) && !SourceDefinitions[DefinitionIndex - 1]->NullOutputs.Contains(Input))
 				{
-					UE_LOG(LogRigMapperDefinition, Warning, TEXT("Could not find matching output in definition %d for input %s in definition %d"), DefinitionIndex - 1, *Input, DefinitionIndex)
+					UE_LOG(LogRigMapper, Warning, TEXT("Could not find matching output in definition %d for input %s in definition %d"), DefinitionIndex - 1, *Input, DefinitionIndex)
 					bOk = false;
 				}
 			}
@@ -1140,7 +1144,7 @@ bool URigMapperLinkedDefinitions::AreLinkedDefinitionsValid() const
 			{
 				if (!SourceDefinitions[DefinitionIndex]->Inputs.Contains(Output.Key))
 				{
-					UE_LOG(LogRigMapperDefinition, Warning, TEXT("Could not find matching input in definition %d for output %s in definition %d"), DefinitionIndex, *Output.Key, DefinitionIndex - 1)
+					UE_LOG(LogRigMapper, Warning, TEXT("Could not find matching input in definition %d for output %s in definition %d"), DefinitionIndex, *Output.Key, DefinitionIndex - 1)
 					bOk = false;
 				}
 			}
diff --git a/Engine/Plugins/Experimental/Animation/RigMapper/Source/RigMapper/Private/RigMapperLog.cpp b/Engine/Plugins/Experimental/Animation/RigMapper/Source/RigMapper/Private/RigMapperLog.cpp
new file mode 100644
--- /dev/null
+++ b/Engine/Plugins/Experimental/Animation/RigMapper/Source/RigMapper/Private/RigMapperLog.cpp
@@ -0,0 +1,5 @@
+// Copyright Epic Games, Inc. All Rights Reserved.
+
+#include "RigMapperLog.h"
+
+DEFINE_LOG_CATEGORY(LogRigMapper);
\ No newline at end of file
diff --git a/Engine/Plugins/Experimental/Animation/RigMapper/Source/RigMapper/Public/RigMapperDefinition.h b/Engine/Plugins/Experimental/Animation/RigMapper/Source/RigMapper/Public/RigMapperDefinition.h
--- a/Engine/Plugins/Experimental/Animation/RigMapper/Source/RigMapper/Public/RigMapperDefinition.h
+++ b/Engine/Plugins/Experimental/Animation/RigMapper/Source/RigMapper/Public/RigMapperDefinition.h
@@ -9,8 +9,6 @@
 #include "RigMapperDefinition.generated.h"
 
 
-DECLARE_LOG_CATEGORY_EXTERN(LogRigMapperDefinition, Log, All);
-
 class FJsonObject;
 class FJsonValue;
 class URigMapperLinkedDefinitions;
@@ -295,12 +293,16 @@ public:
 	
 	UFUNCTION(BlueprintCallable, Category="Animation|Rig Mapper")
 	void Empty();
-	
-	UFUNCTION(BlueprintCallable, Category="Animation|Rig Mapper")
+
+	UFUNCTION(BlueprintCallable, Category = "Animation|Rig Mapper")
 	bool IsDefinitionValid(bool bWarn = false) const;
-	
+
 	UFUNCTION(BlueprintCallable, Category="Animation|Rig Mapper")
 	bool Validate();
+
+	UFUNCTION(BlueprintCallable, Category = "Animation|Rig Mapper")
+	bool IsDefinitionValidated() const;
+
 	
 private:
 	bool LoadInputsFromJsonObject(const TSharedPtr<FJsonObject>& JsonObject);
diff --git a/Engine/Plugins/Experimental/Animation/RigMapper/Source/RigMapper/Public/RigMapperLog.h b/Engine/Plugins/Experimental/Animation/RigMapper/Source/RigMapper/Public/RigMapperLog.h
new file mode 100644
--- /dev/null
+++ b/Engine/Plugins/Experimental/Animation/RigMapper/Source/RigMapper/Public/RigMapperLog.h
@@ -0,0 +1,7 @@
+// Copyright Epic Games, Inc. All Rights Reserved.
+
+#pragma once
+
+#include "Logging/LogMacros.h"
+
+RIGMAPPER_API DECLARE_LOG_CATEGORY_EXTERN(LogRigMapper, Log, All);
diff --git a/Engine/Plugins/Experimental/Animation/RigMapper/Source/RigMapperEditor/Private/RigMapperEditorLog.cpp b/Engine/Plugins/Experimental/Animation/RigMapper/Source/RigMapperEditor/Private/RigMapperEditorLog.cpp
new file mode 100644
--- /dev/null
+++ b/Engine/Plugins/Experimental/Animation/RigMapper/Source/RigMapperEditor/Private/RigMapperEditorLog.cpp
@@ -0,0 +1,5 @@
+// Copyright Epic Games, Inc. All Rights Reserved.
+
+#include "RigMapperEditorLog.h"
+
+DEFINE_LOG_CATEGORY(LogRigMapperEditor);
\ No newline at end of file
diff --git a/Engine/Plugins/Experimental/Animation/RigMapper/Source/RigMapperEditor/Private/RigMapperEditorSubsystem.cpp b/Engine/Plugins/Experimental/Animation/RigMapper/Source/RigMapperEditor/Private/RigMapperEditorSubsystem.cpp
--- a/Engine/Plugins/Experimental/Animation/RigMapper/Source/RigMapperEditor/Private/RigMapperEditorSubsystem.cpp
+++ b/Engine/Plugins/Experimental/Animation/RigMapper/Source/RigMapperEditor/Private/RigMapperEditorSubsystem.cpp
@@ -15,6 +15,7 @@
 #include "Tracks/MovieSceneSpawnTrack.h"
 #include "Units/Execution/RigUnit_PrepareForExecution.h"
 #include "Interfaces//ITargetPlatformManagerModule.h"
+#include "RigMapperEditorLog.h"
 
 UAnimSequence* CreateNewAnimSequence(const FString& BaseName, const FString& Path, USkeletalMesh* TargetMesh)
 {
@@ -60,7 +61,7 @@ UMovieSceneControlRigParameterSection* CreateControlRigSequence(const FString& P
 {
 	if (!SkeletalMesh || !ControlRigClass)
 	{
-		UE_LOG(LogTemp, Error, TEXT("Could not create level sequence and control rig track to record, invalid skeletal mesh or control rig class provided"))
+		UE_LOG(LogRigMapperEditor, Error, TEXT("Could not create level sequence and control rig track to record, invalid skeletal mesh or control rig class provided"))
 		return nullptr;
 	}
 
@@ -91,7 +92,7 @@ UMovieSceneControlRigParameterSection* CreateControlRigSequence(const FString& P
 	const FGuid NewGuid = Sequence->AllowsSpawnableObjects() ? CreateSequencerSpawnable(MovieScene, SkeletalMesh) : FGuid();
 	if (!NewGuid.IsValid())
 	{
-		UE_LOG(LogTemp, Error, TEXT("Failed to create sequencer spawnable from the given skeletal mesh asset"))
+		UE_LOG(LogRigMapperEditor, Error, TEXT("Failed to create sequencer spawnable from the given skeletal mesh asset"))
 		return nullptr;
 	}
 	
@@ -813,7 +814,7 @@ bool URigMapperEditorSubsystem::InitializeRigMapping(const TArray<URigMapperDefi
 {
 	if (Definitions.IsEmpty())
 	{
-		UE_LOG(LogTemp, Error, TEXT("Not enough definitions to remap"))
+		UE_LOG(LogRigMapperEditor, Error, TEXT("Not enough definitions to remap"))
 		return false;
 	}
 	
@@ -823,7 +824,7 @@ bool URigMapperEditorSubsystem::InitializeRigMapping(const TArray<URigMapperDefi
 	{
 		if (!Definitions[i])
 		{
-			UE_LOG(LogTemp, Error, TEXT("Invalid definition at index: %d"), i)
+			UE_LOG(LogRigMapperEditor, Error, TEXT("Invalid definition at index: %d"), i)
 			RigMappers.Empty();
 			return false;
 		}
@@ -831,7 +832,7 @@ bool URigMapperEditorSubsystem::InitializeRigMapping(const TArray<URigMapperDefi
 		
 		if (!RigMapper->LoadDefinition(Definitions[i]))
 		{
-			UE_LOG(LogTemp, Error, TEXT("Could not load definition at index %d"), i)
+			UE_LOG(LogRigMapperEditor, Error, TEXT("Could not load definition at index %d"), i)
 			RigMappers.Empty();
 			return false;
 		}
@@ -863,7 +864,7 @@ void URigMapperEditorSubsystem::EvaluateRigMapping(TArray<FString>& CurveNames,
 			else
 			{
 				// todo: log / interp -> missing input?
-				UE_LOG(LogTemp, Log, TEXT("Missing key, interp not avaialble yet"))
+				UE_LOG(LogRigMapperEditor, Log, TEXT("Missing key, interp not avaialble yet"))
 			}
 		}
 
diff --git a/Engine/Plugins/Experimental/Animation/RigMapper/Source/RigMapperEditor/Public/RigMapperEditorLog.h b/Engine/Plugins/Experimental/Animation/RigMapper/Source/RigMapperEditor/Public/RigMapperEditorLog.h
new file mode 100644
--- /dev/null
+++ b/Engine/Plugins/Experimental/Animation/RigMapper/Source/RigMapperEditor/Public/RigMapperEditorLog.h
@@ -0,0 +1,7 @@
+// Copyright Epic Games, Inc. All Rights Reserved.
+
+#pragma once
+
+#include "Logging/LogMacros.h"
+
+RIGMAPPEREDITOR_API DECLARE_LOG_CATEGORY_EXTERN(LogRigMapperEditor, Log, All);
