Update Rig Mapper Offline Subsystem
Fix framerate issue when generating anim sequence assets.
Update some incorrect return types.

diff --git a/Engine/Plugins/Experimental/Animation/RigMapper/Source/RigMapperEditor/Private/RigMapperEditorSubsystem.cpp b/Engine/Plugins/Experimental/Animation/RigMapper/Source/RigMapperEditor/Private/RigMapperEditorSubsystem.cpp
--- a/Engine/Plugins/Experimental/Animation/RigMapper/Source/RigMapperEditor/Private/RigMapperEditorSubsystem.cpp
+++ b/Engine/Plugins/Experimental/Animation/RigMapper/Source/RigMapperEditor/Private/RigMapperEditorSubsystem.cpp
@@ -19,7 +19,7 @@
 
 UAnimSequence* CreateNewAnimSequence(const FString& BaseName, const FString& Path, USkeletalMesh* TargetMesh)
 {
-	FString UniqueAssetName = BaseName + "_" + TargetMesh->GetName();
+	FString UniqueAssetName = BaseName;
 	
 	FString UniquePackageName;
 	const FString BasePackageName = FString(Path / UniqueAssetName);
@@ -44,7 +44,7 @@ UAnimSequence* CreateNewAnimSequence(const FString& BaseName, const FString& Pat
 	return AnimSequence;
 }
 
-FGuid CreateSequencerSpawnable(UMovieScene* MovieScene, UObject* ObjectToSpawn)
+FGuid CreateSequencerSpawnable(const UMovieScene* MovieScene, UObject* ObjectToSpawn)
 {
 	if (!MovieScene || !ObjectToSpawn)
 	{
@@ -66,7 +66,7 @@ UMovieSceneControlRigParameterSection* CreateControlRigSequence(const FString& P
 	}
 
 	FString UniquePackageName;
-	FString UniqueAssetName = AssetName + "_" + SkeletalMesh->GetName();;
+	FString UniqueAssetName = AssetName;
 	const FString BasePackageName = FString(PackageName / UniqueAssetName);
 
 	const FAssetToolsModule& AssetToolsModule = FModuleManager::Get().LoadModuleChecked<FAssetToolsModule>("AssetTools");
@@ -143,7 +143,7 @@ TArray<UMovieSceneControlRigParameterSection*> URigMapperEditorSubsystem::GetSec
 	{
 		for (const FMovieSceneBinding& Binding : Sequence->GetMovieScene()->GetBindings())
 		{
-			for (UMovieSceneTrack* Track : Binding.GetTracks())
+			for (const UMovieSceneTrack* Track : Binding.GetTracks())
 			{
 				if (Track->IsA(UMovieSceneControlRigParameterTrack::StaticClass()))
 				{
@@ -177,8 +177,8 @@ void URigMapperEditorSubsystem::SetAnimSequenceRate(UAnimSequence* AnimSequence,
 		Controller.OpenBracket(INVTEXT("Set Sequence FrameRate"));
 		Controller.InitializeModel();
 		Controller.SetFrameRate(FrameRate, false);
-
-		Controller.CloseBracket();	
+		Controller.NotifyPopulated();
+		Controller.CloseBracket();
 	
 		AnimSequence->MarkPackageDirty();
 	}
@@ -255,13 +255,16 @@ bool URigMapperEditorSubsystem::AddCurveValuesToAnimSequence(UAnimSequence* Targ
 	{
 		return false;
 	}
-
-	Target->ResetAnimation();
+	
 	
 	IAnimationDataController& Controller = Target->GetController();
 	Controller.OpenBracket(INVTEXT("Populate New Anim"));
 	Controller.InitializeModel();
-	
+	Controller.RemoveAllAttributes(false);
+	Controller.RemoveAllBoneTracks(false);
+	Controller.RemoveAllCurvesOfType(ERawCurveTrackTypes::RCT_Float, false);
+	Controller.RemoveAllCurvesOfType(ERawCurveTrackTypes::RCT_Transform, false);
+	Controller.RemoveAllCurvesOfType(ERawCurveTrackTypes::RCT_Vector, false);
 	Controller.SetNumberOfFrames(FrameTimes.Last().CeilToFrame().Value - FrameTimes[0].FloorToFrame().Value + 1, false);
 
 	// AnimSequence should start on frame 0
@@ -347,9 +350,9 @@ bool URigMapperEditorSubsystem::ConvertAnimSequence(const UAnimSequence* Source,
 	return ConvertCurveValuesToAnimSequence(CurveNames, FrameTimes, CurveValuesPerFrame, Target, Definitions);
 }
 
-bool URigMapperEditorSubsystem::ConvertAnimSequenceNew(const UAnimSequence* Source, USkeletalMesh* TargetMesh, const TArray<URigMapperDefinition*>& Definitions)
+UAnimSequence* URigMapperEditorSubsystem::ConvertAnimSequenceNew(const UAnimSequence* Source, USkeletalMesh* TargetMesh, const TArray<URigMapperDefinition*>& Definitions, const FDirectoryPath& NewAssetPath, const FName NewAssetName)
 {
-	UAnimSequence* NewSequence = CreateNewAnimSequence(Source->GetName(), FPaths::GetPath(Source->GetPathName()), TargetMesh);
+	UAnimSequence* NewSequence = CreateNewAnimSequence(NewAssetName.ToString(), NewAssetPath.Path, TargetMesh);
 
 	NewSequence->ImportFileFramerate = Source->ImportFileFramerate;
 	NewSequence->ImportResampleFramerate = Source->ImportResampleFramerate;
@@ -357,10 +360,10 @@ bool URigMapperEditorSubsystem::ConvertAnimSequenceNew(const UAnimSequence* Sour
 	
 	if (!ConvertAnimSequence(Source, NewSequence, Definitions))
 	{
-		return false;
+		return nullptr;
 	}
 	NewSequence->PostLoad();
-	return true;
+	return NewSequence;
 }
 
 bool URigMapperEditorSubsystem::ConvertCsvToAnimSequence(const FFilePath& InputFile, UAnimSequence* Target, const TArray<URigMapperDefinition*>& Definitions)
@@ -377,18 +380,18 @@ bool URigMapperEditorSubsystem::ConvertCsvToAnimSequence(const FFilePath& InputF
 	return ConvertCurveValuesToAnimSequence(CurveNames, FrameTimes, CurveValuesPerFrame, Target, Definitions);
 }
 
-bool URigMapperEditorSubsystem::ConvertCsvToAnimSequenceNew(const FFilePath& InputFile, const FFilePath& NewAssetPath, USkeletalMesh* TargetMesh, const TArray<URigMapperDefinition*>& Definitions, const FFrameRate& FrameRate)
+UAnimSequence* URigMapperEditorSubsystem::ConvertCsvToAnimSequenceNew(const FFilePath& InputFile, USkeletalMesh* TargetMesh, const TArray<URigMapperDefinition*>& Definitions, const FFrameRate& FrameRate, const FDirectoryPath& NewAssetPath, const FName NewAssetName)
 {
-	UAnimSequence* NewSequence = CreateNewAnimSequence(FPaths::GetBaseFilename(InputFile.FilePath), NewAssetPath.FilePath, TargetMesh);
+	UAnimSequence* NewSequence = CreateNewAnimSequence(NewAssetName.ToString(), NewAssetPath.Path, TargetMesh);
 	
 	SetAnimSequenceRate(NewSequence, FrameRate, true);
 
 	if (!ConvertCsvToAnimSequence(InputFile, NewSequence, Definitions))
 	{
-		return false;
+		return nullptr;
 	}
 	NewSequence->PostLoad();
-	return true;
+	return NewSequence;
 }
 
 bool URigMapperEditorSubsystem::ConvertAnimSequenceToCsv(const UAnimSequence* Source, const FFilePath& OutputFile, const TArray<URigMapperDefinition*>& Definitions, bool bOutputIntermediateCsvFiles)
@@ -419,9 +422,9 @@ bool URigMapperEditorSubsystem::ConvertAnimSequenceToControlRigSection(const UAn
 	return ConvertCurveValuesToControlRigSection(CurveNames, FrameTimes, CurveValuesPerFrame, Target, Definitions);
 }
 
-bool URigMapperEditorSubsystem::ConvertAnimSequenceToControlRigSectionNew(const UAnimSequence* Source, USkeletalMesh* TargetMesh, const TArray<URigMapperDefinition*>& Definitions, const TSubclassOf<UControlRig>& ControlRigClass)
+UMovieSceneControlRigParameterSection* URigMapperEditorSubsystem::ConvertAnimSequenceToControlRigSectionNew(const UAnimSequence* Source, USkeletalMesh* TargetMesh, const TArray<URigMapperDefinition*>& Definitions, const TSubclassOf<UControlRig>& ControlRigClass, const FDirectoryPath& NewAssetPath, const FName NewAssetName)
 {
-	UMovieSceneControlRigParameterSection* NewSection = CreateControlRigSequence(FPaths::GetPath(Source->GetPathName()), Source->GetName(), Source->GetTypedOuter<UMovieScene>()->GetDisplayRate(), ControlRigClass, TargetMesh);
+	UMovieSceneControlRigParameterSection* NewSection = CreateControlRigSequence(NewAssetPath.Path, NewAssetName.ToString(), Source->GetTypedOuter<UMovieScene>()->GetDisplayRate(), ControlRigClass, TargetMesh);
 	
 	TArray<FString> CurveNames;
 	TArray<FFrameTime> FrameTimes;
@@ -429,10 +432,14 @@ bool URigMapperEditorSubsystem::ConvertAnimSequenceToControlRigSectionNew(const
 	
 	if (!LoadCurveValuesFromAnimSequence(Source, CurveNames, FrameTimes, CurveValuesPerFrame))
 	{
-		return false;
+		return nullptr;
 	}
 	
-	return ConvertCurveValuesToControlRigSection(CurveNames, FrameTimes, CurveValuesPerFrame, NewSection, Definitions);
+	if (ConvertCurveValuesToControlRigSection(CurveNames, FrameTimes, CurveValuesPerFrame, NewSection, Definitions))
+	{
+		return NewSection;
+	}
+	return nullptr;
 }
 
 bool URigMapperEditorSubsystem::ConvertCurveValuesToControlRigSection(TArray<FString>& CurveNames, const TArray<FFrameTime>& FrameTimes, TArray<FPoseValues>& CurveValuesPerFrame, UMovieSceneControlRigParameterSection* Target, const TArray<URigMapperDefinition*>& Definitions)
@@ -651,11 +658,15 @@ bool URigMapperEditorSubsystem::ConvertControlRigSection(const UMovieSceneContro
 	return ConvertCurveValuesToControlRigSection(CurveNames, FrameTimes, CurveValuesPerFrame, Target, Definitions);
 }
 
-bool URigMapperEditorSubsystem::ConvertControlRigSectionNew(const UMovieSceneControlRigParameterSection* Source, USkeletalMesh* TargetMesh, const TArray<URigMapperDefinition*>& Definitions, const TSubclassOf<UControlRig>& ControlRigClass)
+UMovieSceneControlRigParameterSection* URigMapperEditorSubsystem::ConvertControlRigSectionNew(const UMovieSceneControlRigParameterSection* Source, USkeletalMesh* TargetMesh, const TArray<URigMapperDefinition*>& Definitions, const TSubclassOf<UControlRig>& ControlRigClass, const FDirectoryPath& NewAssetPath, const FName NewAssetName)
 {
-	UMovieSceneControlRigParameterSection* NewSection = CreateControlRigSequence(FPaths::GetPath(Source->GetPathName()), Source->GetName(), Source->GetTypedOuter<UMovieScene>()->GetDisplayRate(), ControlRigClass, TargetMesh);
+	UMovieSceneControlRigParameterSection* NewSection = CreateControlRigSequence(NewAssetPath.Path, NewAssetName.ToString(), Source->GetTypedOuter<UMovieScene>()->GetDisplayRate(), ControlRigClass, TargetMesh);
 	
-	return ConvertControlRigSection(Source, NewSection, Definitions);
+	if (ConvertControlRigSection(Source, NewSection, Definitions))
+	{
+		return NewSection;
+	}
+	return nullptr;
 }
 
 bool URigMapperEditorSubsystem::ConvertCsvToControlRigSection(const FFilePath& InputFile, UMovieSceneControlRigParameterSection* Target, const TArray<URigMapperDefinition*>& Definitions)
@@ -672,11 +683,15 @@ bool URigMapperEditorSubsystem::ConvertCsvToControlRigSection(const FFilePath& I
 	return ConvertCurveValuesToControlRigSection(CurveNames, FrameTimes, CurveValuesPerFrame, Target, Definitions);
 }
 
-bool URigMapperEditorSubsystem::ConvertCsvToControlRigSectionNew(const FFilePath& InputFile, USkeletalMesh* TargetMesh, const TArray<URigMapperDefinition*>& Definitions, const FFilePath& NewAssetPath, const FFrameRate& FrameRate, const TSubclassOf<UControlRig>& ControlRigClass)
+UMovieSceneControlRigParameterSection* URigMapperEditorSubsystem::ConvertCsvToControlRigSectionNew(const FFilePath& InputFile, USkeletalMesh* TargetMesh, const TArray<URigMapperDefinition*>& Definitions, const FFrameRate& FrameRate, const TSubclassOf<UControlRig>& ControlRigClass, const FDirectoryPath& NewAssetPath, const FName NewAssetName)
 {
-	UMovieSceneControlRigParameterSection* NewSection = CreateControlRigSequence(FPaths::GetPath(NewAssetPath.FilePath), FPaths::GetBaseFilename(InputFile.FilePath), FrameRate, ControlRigClass, TargetMesh);
+	UMovieSceneControlRigParameterSection* NewSection = CreateControlRigSequence(NewAssetPath.Path, NewAssetName.ToString(), FrameRate, ControlRigClass, TargetMesh);
 
-	return ConvertCsvToControlRigSection(InputFile, NewSection, Definitions);
+	if (ConvertCsvToControlRigSection(InputFile, NewSection, Definitions))
+	{
+		return NewSection;
+	}
+	return nullptr;
 }
 
 bool URigMapperEditorSubsystem::ConvertControlRigSectionToCsv(const UMovieSceneControlRigParameterSection* Source, const FFilePath& OutputFile, const TArray<URigMapperDefinition*>& Definitions, bool bOutputIntermediateCsvFiles)
@@ -715,7 +730,7 @@ bool URigMapperEditorSubsystem::ConvertControlRigSectionToAnimSequence(const UMo
 	return ConvertCurveValuesToAnimSequence(CurveNames, FrameTimes, CurveValuesPerFrame, Target, Definitions);
 }
 
-bool URigMapperEditorSubsystem::ConvertControlRigSectionToAnimSequenceNew(const UMovieSceneControlRigParameterSection* Source, USkeletalMesh* TargetMesh, const TArray<URigMapperDefinition*>& Definitions)
+bool URigMapperEditorSubsystem::ConvertControlRigSectionToAnimSequenceNew(const UMovieSceneControlRigParameterSection* Source, USkeletalMesh* TargetMesh, const TArray<URigMapperDefinition*>& Definitions, const FDirectoryPath& NewAssetPath, const FName NewAssetName)
 {
 	ULevelSequence* Sequence = GetSequenceFromSection(Source);
 	if (!Sequence || !Sequence->GetMovieScene())
@@ -723,7 +738,7 @@ bool URigMapperEditorSubsystem::ConvertControlRigSectionToAnimSequenceNew(const
 		return false;
 	}
 	
-	UAnimSequence* NewSequence = CreateNewAnimSequence(Sequence->GetName(), FPaths::GetPath(Sequence->GetPathName()), TargetMesh);
+	UAnimSequence* NewSequence = CreateNewAnimSequence(NewAssetName.ToString(), NewAssetPath.Path, TargetMesh);
 	SetAnimSequenceRate(NewSequence, Sequence->GetMovieScene()->GetDisplayRate(), true);
 	
 	TArray<FString> CurveNames;
diff --git a/Engine/Plugins/Experimental/Animation/RigMapper/Source/RigMapperEditor/Public/RigMapperEditorSubsystem.h b/Engine/Plugins/Experimental/Animation/RigMapper/Source/RigMapperEditor/Public/RigMapperEditorSubsystem.h
--- a/Engine/Plugins/Experimental/Animation/RigMapper/Source/RigMapperEditor/Public/RigMapperEditorSubsystem.h
+++ b/Engine/Plugins/Experimental/Animation/RigMapper/Source/RigMapperEditor/Public/RigMapperEditorSubsystem.h
@@ -54,7 +54,7 @@ public:
 	 * The following header is expected at the first line of the CSV file: curve_name, frame_number, value
 	 */
 	UFUNCTION(BlueprintCallable, Category = "Editor Scripting|Animation|Rig Mapper")
-	static bool ConvertCsvToAnimSequenceNew(const FFilePath& InputFile, const FFilePath& NewAssetPath, USkeletalMesh* TargetMesh, const TArray<URigMapperDefinition*>& Definitions, const FFrameRate& FrameRate);
+	static UAnimSequence* ConvertCsvToAnimSequenceNew(const FFilePath& InputFile, USkeletalMesh* TargetMesh, const TArray<URigMapperDefinition*>& Definitions, const FFrameRate& FrameRate, const FDirectoryPath& NewAssetPath, const FName NewAssetName);
 
 	/**
 	 * Import a CSV file to the given control rig section, remapping it using the given definitions
@@ -71,7 +71,7 @@ public:
 	 * Control Rig transform/vector2d/float channels are currently supported
 	 */
 	UFUNCTION(BlueprintCallable, Category = "Editor Scripting|Animation|Rig Mapper")
-	static bool ConvertCsvToControlRigSectionNew(const FFilePath& InputFile, USkeletalMesh* TargetMesh, const TArray<URigMapperDefinition*>& Definitions, const FFilePath& NewAssetPath, const FFrameRate& FrameRate, const TSubclassOf<UControlRig>& ControlRigClass);
+	static UMovieSceneControlRigParameterSection* ConvertCsvToControlRigSectionNew(const FFilePath& InputFile, USkeletalMesh* TargetMesh, const TArray<URigMapperDefinition*>& Definitions, const FFrameRate& FrameRate, const TSubclassOf<UControlRig>& ControlRigClass, const FDirectoryPath& NewAssetPath, const FName NewAssetName);
 
 	/**
 	 * Remap an anim sequence to another existing one, remapping it using the given definitions
@@ -84,7 +84,7 @@ public:
 	 * Remap an anim sequence to a new one, remapping it using the given definitions
 	 */
 	UFUNCTION(BlueprintCallable, Category = "Editor Scripting|Animation|Rig Mapper")
-	static bool ConvertAnimSequenceNew(const UAnimSequence* Source, USkeletalMesh* TargetMesh, const TArray<URigMapperDefinition*>& Definitions);
+	static UAnimSequence* ConvertAnimSequenceNew(const UAnimSequence* Source, USkeletalMesh* TargetMesh, const TArray<URigMapperDefinition*>& Definitions, const FDirectoryPath& NewAssetPath, const FName NewAssetName);
 
 	/**
 	 * Export a CSV file from an anim sequence, remapping it using the given definitions
@@ -105,7 +105,7 @@ public:
 	 * Control Rig transform/vector2d/float channels are currently supported
 	 */
 	UFUNCTION(BlueprintCallable, Category = "Editor Scripting|Animation|Rig Mapper")
-	static bool ConvertAnimSequenceToControlRigSectionNew(const UAnimSequence* Source, USkeletalMesh* TargetMesh, const TArray<URigMapperDefinition*>& Definitions, const TSubclassOf<UControlRig>& ControlRigClass);
+	static UMovieSceneControlRigParameterSection* ConvertAnimSequenceToControlRigSectionNew(const UAnimSequence* Source, USkeletalMesh* TargetMesh, const TArray<URigMapperDefinition*>& Definitions, const TSubclassOf<UControlRig>& ControlRigClass, const FDirectoryPath& NewAssetPath, const FName NewAssetName);
 
 	/**
 	 * Convert a Control Rig section to another existing one, remapping it using the given definitions
@@ -120,7 +120,7 @@ public:
 	 * Control Rig transform/vector2d/float channels are currently supported
 	 */
 	UFUNCTION(BlueprintCallable, Category = "Editor Scripting|Animation|Rig Mapper")
-	static bool ConvertControlRigSectionNew(const UMovieSceneControlRigParameterSection* Source, USkeletalMesh* TargetMesh, const TArray<URigMapperDefinition*>& Definitions, const TSubclassOf<UControlRig>& ControlRigClass);
+	static UMovieSceneControlRigParameterSection* ConvertControlRigSectionNew(const UMovieSceneControlRigParameterSection* Source, USkeletalMesh* TargetMesh, const TArray<URigMapperDefinition*>& Definitions, const TSubclassOf<UControlRig>& ControlRigClass, const FDirectoryPath& NewAssetPath, const FName NewAssetName);
 
 	/**
 	 * Export a CSV file from a Control Rig Section, remapping it using the given definitions
@@ -139,7 +139,7 @@ public:
 	 * Convert a Control Rig section to a new AnimSequence asset, remapping it using the given definitions
 	 */
 	UFUNCTION(BlueprintCallable, Category = "Editor Scripting|Animation|Rig Mapper")
-	static bool ConvertControlRigSectionToAnimSequenceNew(const UMovieSceneControlRigParameterSection* Source, USkeletalMesh* TargetMesh, const TArray<URigMapperDefinition*>& Definitions);
+	static bool ConvertControlRigSectionToAnimSequenceNew(const UMovieSceneControlRigParameterSection* Source, USkeletalMesh* TargetMesh, const TArray<URigMapperDefinition*>& Definitions, const FDirectoryPath& NewAssetPath, const FName NewAssetName);
 	
 	UFUNCTION(BlueprintCallable, Category = "Editor Scripting|Animation|Rig Mapper")
 	static ULevelSequence* GetSequenceFromSection(const UMovieSceneControlRigParameterSection* Section);
