add visualization of NullOutputs in Structure View and Graph View

diff --git a/Engine/Plugins/Experimental/Animation/RigMapper/Source/RigMapperEditor/Private/RigMapperDefinitionEditorToolkit.cpp b/Engine/Plugins/Experimental/Animation/RigMapper/Source/RigMapperEditor/Private/RigMapperDefinitionEditorToolkit.cpp
--- a/Engine/Plugins/Experimental/Animation/RigMapper/Source/RigMapperEditor/Private/RigMapperDefinitionEditorToolkit.cpp
+++ b/Engine/Plugins/Experimental/Animation/RigMapper/Source/RigMapperEditor/Private/RigMapperDefinitionEditorToolkit.cpp
@@ -29,6 +29,7 @@ const TMap<FName, ERigMapperNodeType> FRigMapperDefinitionEditorToolkit::Propert
 	{ GET_MEMBER_NAME_CHECKED(FRigMapperFeatureDefinitions, WeightedSums), ERigMapperNodeType::WeightedSum },
 	{ GET_MEMBER_NAME_CHECKED(FRigMapperFeatureDefinitions, SDKs), ERigMapperNodeType::SDK },
 	{ GET_MEMBER_NAME_CHECKED(URigMapperDefinition, Outputs), ERigMapperNodeType::Output },
+	{ GET_MEMBER_NAME_CHECKED(URigMapperDefinition, NullOutputs), ERigMapperNodeType::NullOutput },
 };
 
 void FRigMapperDefinitionEditorToolkit::Initialize(URigMapperDefinition* InDefinition, EToolkitMode::Type InMode, TSharedPtr<IToolkitHost> InToolkitHost)
@@ -240,11 +241,6 @@ bool FRigMapperDefinitionEditorToolkit::HandleIsPropertyVisible(const FPropertyA
 	{
 		return StructureView->IsNodeOrChildSelected(PropertyNameToNodeTypeMapping[PropertyName], ArrayIndex);
 	}
-	if (PropertyName == GET_MEMBER_NAME_CHECKED(URigMapperDefinition, NullOutputs))
-	{
-		// todo: add a way to select the null outputs from the structure
-		return false;
-	}
 	return true;
 }
 
@@ -260,15 +256,15 @@ void FRigMapperDefinitionEditorToolkit::HandleFinishedChangingProperties(const F
 	const FName OutputName = GET_MEMBER_NAME_CHECKED(URigMapperDefinition, Outputs);
 	const FName FeaturesName = GET_MEMBER_NAME_CHECKED(URigMapperDefinition, Features);
 	
+	const FName NullOutputName = GET_MEMBER_NAME_CHECKED(URigMapperDefinition, NullOutputs);
+
 	const bool bInput = PropertyName == InputName && ParentName == InputName;
 	const bool bOutput = !bInput && (PropertyName == OutputName || PropertyName == OutputName.ToString() + TEXT("_Key")) && ParentName == OutputName; // PropertyName will be Outputs_Key if renaming an output
 	const bool bName = !bInput && !bOutput && PropertyName == GET_MEMBER_NAME_CHECKED(FRigMapperFeature, Name);
 	const bool bFeaturesChildProperty = ParentName == FeaturesName;
-
 	const bool bSelectNew = (ChangeType == EPropertyChangeType::Duplicate || ChangeType == EPropertyChangeType::ArrayAdd);
-
-
-	const bool bRename = ChangeType == EPropertyChangeType::ValueSet && (bInput || bOutput || bName || bFeaturesChildProperty);
+	const bool bNullOutput = PropertyName == NullOutputName && ParentName == NullOutputName;
+	const bool bRename = ChangeType == EPropertyChangeType::ValueSet && (bInput || bOutput || bName || bFeaturesChildProperty || bNullOutput);
 	
 	GraphEditor->RebuildGraph();
 	StructureView->RebuildTree();
@@ -279,7 +275,7 @@ void FRigMapperDefinitionEditorToolkit::HandleFinishedChangingProperties(const F
 		{
 			int32 InputIndex = Event.GetArrayIndex(PropertyName.ToString());
 			StructureView->SelectNode(Definition->Inputs[InputIndex], ERigMapperNodeType::Input, true);
-			GraphEditor->SelectNodes({ Definition->Inputs[InputIndex] }, {}, {});
+			GraphEditor->SelectNodes({ Definition->Inputs[InputIndex] }, {}, {}, {});
 		}
 		else if (bOutput)
 		{
@@ -287,24 +283,30 @@ void FRigMapperDefinitionEditorToolkit::HandleFinishedChangingProperties(const F
 			TArray<FString> OutputsInOrder;
 			Definition->Outputs.GenerateKeyArray(OutputsInOrder);
 			StructureView->SelectNode(OutputsInOrder[OutputIndex], ERigMapperNodeType::Output, true);
-			GraphEditor->SelectNodes({}, {}, { OutputsInOrder[OutputIndex] });
+			GraphEditor->SelectNodes({}, {}, { OutputsInOrder[OutputIndex] }, {});
+		}
+		if (bNullOutput)
+		{
+			int32 NullOutputIndex = Event.GetArrayIndex(PropertyName.ToString());
+			StructureView->SelectNode(Definition->NullOutputs[NullOutputIndex], ERigMapperNodeType::NullOutput, true);
+			GraphEditor->SelectNodes({ Definition->NullOutputs[NullOutputIndex] }, {}, {}, {});
 		}
 		else if (bFeaturesChildProperty)
 		{
 			if (const int32 MultiplyIndex = Event.GetArrayIndex(GET_MEMBER_NAME_CHECKED(FRigMapperFeatureDefinitions, Multiply).ToString()); MultiplyIndex != INDEX_NONE)
 			{
 				StructureView->SelectNode(Definition->Features.Multiply[MultiplyIndex].Name, ERigMapperNodeType::Multiply, true);
-				GraphEditor->SelectNodes({}, { Definition->Features.Multiply[MultiplyIndex].Name }, {});
+				GraphEditor->SelectNodes({}, { Definition->Features.Multiply[MultiplyIndex].Name }, {}, {});
 			}
 			else if (const int32 WsIndex = Event.GetArrayIndex(GET_MEMBER_NAME_CHECKED(FRigMapperFeatureDefinitions, WeightedSums).ToString()); WsIndex != INDEX_NONE)
 			{
 				StructureView->SelectNode(Definition->Features.WeightedSums[WsIndex].Name, ERigMapperNodeType::WeightedSum, true);
-				GraphEditor->SelectNodes({}, { Definition->Features.WeightedSums[WsIndex].Name }, { });
+				GraphEditor->SelectNodes({}, { Definition->Features.WeightedSums[WsIndex].Name }, { }, {});
 			}
 			else if (const int32 SdkIndex = Event.GetArrayIndex(GET_MEMBER_NAME_CHECKED(FRigMapperFeatureDefinitions, SDKs).ToString()); SdkIndex != INDEX_NONE)
 			{
 				StructureView->SelectNode(Definition->Features.SDKs[SdkIndex].Name, ERigMapperNodeType::SDK, true);
-				GraphEditor->SelectNodes({}, { Definition->Features.SDKs[SdkIndex].Name }, {});
+				GraphEditor->SelectNodes({}, { Definition->Features.SDKs[SdkIndex].Name }, {}, {});
 			}
 		}
 	}
@@ -329,13 +331,13 @@ void FRigMapperDefinitionEditorToolkit::HandleGraphSelectionChanged(const TSet<U
 	DetailsView->ForceRefresh();
 }
 
-void FRigMapperDefinitionEditorToolkit::HandleStructureSelectionChanged(ESelectInfo::Type SelectInfo, TArray<FString> SelectedInputs, TArray<FString> SelectedFeatures, TArray<FString> SelectedOutputs)
+void FRigMapperDefinitionEditorToolkit::HandleStructureSelectionChanged(ESelectInfo::Type SelectInfo, TArray<FString> SelectedInputs, TArray<FString> SelectedFeatures, TArray<FString> SelectedOutputs, TArray<FString> SelectedNullOutputs)
 {
 	DetailsView->ForceRefresh();
 
 	if (SelectInfo != ESelectInfo::Direct)
 	{
-		GraphEditor->SelectNodes(SelectedInputs, SelectedFeatures, SelectedOutputs);
+		GraphEditor->SelectNodes(SelectedInputs, SelectedFeatures, SelectedOutputs, SelectedNullOutputs);
 	};
 }
 
diff --git a/Engine/Plugins/Experimental/Animation/RigMapper/Source/RigMapperEditor/Private/RigMapperDefinitionEditorToolkit.h b/Engine/Plugins/Experimental/Animation/RigMapper/Source/RigMapperEditor/Private/RigMapperDefinitionEditorToolkit.h
--- a/Engine/Plugins/Experimental/Animation/RigMapper/Source/RigMapperEditor/Private/RigMapperDefinitionEditorToolkit.h
+++ b/Engine/Plugins/Experimental/Animation/RigMapper/Source/RigMapperEditor/Private/RigMapperDefinitionEditorToolkit.h
@@ -53,7 +53,7 @@ private:
 	void HandleRigMapperDefinitionLoaded();
 	void HandleGraphSelectionChanged(const TSet<UObject*>& Nodes);
 	
-	void HandleStructureSelectionChanged(ESelectInfo::Type SelectInfo, TArray<FString> SelectedInputs, TArray<FString> SelectedFeatures, TArray<FString> SelectedOutputs);
+	void HandleStructureSelectionChanged(ESelectInfo::Type SelectInfo, TArray<FString> SelectedInputs, TArray<FString> SelectedFeatures, TArray<FString> SelectedOutputs, TArray<FString> SelectedNullOutputs);
 
 private:
 	static const FName DefinitionEditorGraphTabId;
diff --git a/Engine/Plugins/Experimental/Animation/RigMapper/Source/RigMapperEditor/Private/RigMapperGraph/RigMapperDefinitionEditorGraph.cpp b/Engine/Plugins/Experimental/Animation/RigMapper/Source/RigMapperEditor/Private/RigMapperGraph/RigMapperDefinitionEditorGraph.cpp
--- a/Engine/Plugins/Experimental/Animation/RigMapper/Source/RigMapperEditor/Private/RigMapperGraph/RigMapperDefinitionEditorGraph.cpp
+++ b/Engine/Plugins/Experimental/Animation/RigMapper/Source/RigMapperEditor/Private/RigMapperGraph/RigMapperDefinitionEditorGraph.cpp
@@ -54,12 +54,19 @@ void URigMapperDefinitionEditorGraph::ConstructNodes()
 			CreateGraphNode(Input, ERigMapperNodeType::Input);
 		}
 	}
+	for (const FString& NullOutput : Definition->NullOutputs)
+	{
+		if (!NullOutputNodes.Contains(NullOutput))
+		{
+			CreateGraphNode(NullOutput, ERigMapperNodeType::NullOutput);
+		}
+	}
 }
 
-TArray<URigMapperDefinitionEditorGraphNode*> URigMapperDefinitionEditorGraph::GetNodesByName(const TArray<FString>& Inputs, const TArray<FString>& Features, const TArray<FString>& Outputs) const
+TArray<URigMapperDefinitionEditorGraphNode*> URigMapperDefinitionEditorGraph::GetNodesByName(const TArray<FString>& Inputs, const TArray<FString>& Features, const TArray<FString>& Outputs, const TArray<FString>& NullOutputs) const
 {
 	TArray<URigMapperDefinitionEditorGraphNode*> OutNodes;
-	OutNodes.Reserve(Inputs.Num() + Features.Num() + Outputs.Num());
+	OutNodes.Reserve(Inputs.Num() + Features.Num() + Outputs.Num() + NullOutputs.Num());
 
 	auto AddNodesByNameFromMap = [&OutNodes](const TArray<FString>& InArray, TMap<FString, URigMapperDefinitionEditorGraphNode*> InMap)
 	{
@@ -75,6 +82,7 @@ TArray<URigMapperDefinitionEditorGraphNode*> URigMapperDefinitionEditorGraph::Ge
 	AddNodesByNameFromMap(Inputs, InputNodes);
 	AddNodesByNameFromMap(Features, FeatureNodes);
 	AddNodesByNameFromMap(Outputs, OutputNodes);
+	AddNodesByNameFromMap(NullOutputs, NullOutputNodes);
 
 	return OutNodes;
 }
@@ -192,6 +200,10 @@ URigMapperDefinitionEditorGraphNode* URigMapperDefinitionEditorGraph::CreateGrap
 	{
 		OutputNodes.Add(NodeName, Node);
 	}
+	else if (NodeType == ERigMapperNodeType::NullOutput)
+	{
+		NullOutputNodes.Add(NodeName, Node);
+	}
 	else // todo: Add feature type
 	{
 		FeatureNodes.Add(NodeName, Node);
@@ -244,6 +256,10 @@ void URigMapperDefinitionEditorGraph::LayoutNodeRec(URigMapperDefinitionEditorGr
 	{
 		Margin.X = 0;
 	}
+	else if (InNode->GetNodeType() == ERigMapperNodeType::NullOutput)
+	{
+		Margin.X = 0;
+	}
 	else
 	{
 		InNode->NodePosX = FMath::Max(InNode->NodePosX, InputsWidth);
@@ -326,6 +342,16 @@ void URigMapperDefinitionEditorGraph::LayoutNodes() const
 			PosY = Node.Value->NodePosY + Node.Value->GetDimensions().Y + Node.Value->GetMargin().Y + SubGraphMarginY;
 		}
 	}
+
+	for (const TPair<FString, URigMapperDefinitionEditorGraphNode*>& Node : NullOutputNodes)
+	{
+		if (!LayedOutNodes.Contains(Node.Value))
+		{
+			Node.Value->NodePosX = MaxPosX + OutputMarginX;
+			LayoutNodeRec(Node.Value, InputsMaxWidth, PosY, LayedOutNodes);
+			PosY = Node.Value->NodePosY + Node.Value->GetDimensions().Y + Node.Value->GetMargin().Y + SubGraphMarginY;
+		}
+	}
 }
 
 void URigMapperDefinitionEditorGraph::RemoveAllNodes()
@@ -341,4 +367,5 @@ void URigMapperDefinitionEditorGraph::RemoveAllNodes()
 	InputNodes.Reset();
 	FeatureNodes.Reset();
 	OutputNodes.Reset();
+	NullOutputNodes.Reset();
 }
diff --git a/Engine/Plugins/Experimental/Animation/RigMapper/Source/RigMapperEditor/Private/RigMapperGraph/RigMapperDefinitionEditorGraph.h b/Engine/Plugins/Experimental/Animation/RigMapper/Source/RigMapperEditor/Private/RigMapperGraph/RigMapperDefinitionEditorGraph.h
--- a/Engine/Plugins/Experimental/Animation/RigMapper/Source/RigMapperEditor/Private/RigMapperGraph/RigMapperDefinitionEditorGraph.h
+++ b/Engine/Plugins/Experimental/Animation/RigMapper/Source/RigMapperEditor/Private/RigMapperGraph/RigMapperDefinitionEditorGraph.h
@@ -30,7 +30,7 @@ public:
 	
 	URigMapperDefinition* GetDefinition() const { return WeakDefinition.Get(); };
 
-	TArray<URigMapperDefinitionEditorGraphNode*> GetNodesByName(const TArray<FString>& Inputs, const TArray<FString>& Features, const TArray<FString>& Outputs) const;
+	TArray<URigMapperDefinitionEditorGraphNode*> GetNodesByName(const TArray<FString>& Inputs, const TArray<FString>& Features, const TArray<FString>& Outputs, const TArray<FString>& NullOutputs) const;
 
 private:
 	URigMapperDefinitionEditorGraphNode* CreateGraphNodesRec(URigMapperDefinition* Definition, const FString& InNodeName, bool bIsOutputNode);
@@ -45,6 +45,7 @@ private:
 	TMap<FString, URigMapperDefinitionEditorGraphNode*> InputNodes;
 	TMap<FString, URigMapperDefinitionEditorGraphNode*> FeatureNodes;
 	TMap<FString, URigMapperDefinitionEditorGraphNode*> OutputNodes;
+	TMap<FString, URigMapperDefinitionEditorGraphNode*> NullOutputNodes;
 
 	bool bRefreshLayout = false;
 };
diff --git a/Engine/Plugins/Experimental/Animation/RigMapper/Source/RigMapperEditor/Private/RigMapperGraph/RigMapperDefinitionEditorGraphNode.cpp b/Engine/Plugins/Experimental/Animation/RigMapper/Source/RigMapperEditor/Private/RigMapperGraph/RigMapperDefinitionEditorGraphNode.cpp
--- a/Engine/Plugins/Experimental/Animation/RigMapper/Source/RigMapperEditor/Private/RigMapperGraph/RigMapperDefinitionEditorGraphNode.cpp
+++ b/Engine/Plugins/Experimental/Animation/RigMapper/Source/RigMapperEditor/Private/RigMapperGraph/RigMapperDefinitionEditorGraphNode.cpp
@@ -40,6 +40,10 @@ FLinearColor URigMapperDefinitionEditorGraphNode::GetNodeTitleColor() const
 	{
 		return FLinearColor(0.466, 0.969, 0.525);
 	}
+	if (NodeType == ERigMapperNodeType::NullOutput)
+	{
+		return FLinearColor(0.65, 0.969, 0.466);
+	}
 	if (NodeType == ERigMapperNodeType::Multiply)
 	{
 		return FLinearColor(0.969, 0.966, 0.466);
diff --git a/Engine/Plugins/Experimental/Animation/RigMapper/Source/RigMapperEditor/Private/RigMapperGraph/RigMapperDefinitionEditorGraphNode.h b/Engine/Plugins/Experimental/Animation/RigMapper/Source/RigMapperEditor/Private/RigMapperGraph/RigMapperDefinitionEditorGraphNode.h
--- a/Engine/Plugins/Experimental/Animation/RigMapper/Source/RigMapperEditor/Private/RigMapperGraph/RigMapperDefinitionEditorGraphNode.h
+++ b/Engine/Plugins/Experimental/Animation/RigMapper/Source/RigMapperEditor/Private/RigMapperGraph/RigMapperDefinitionEditorGraphNode.h
@@ -18,6 +18,7 @@ enum class ERigMapperNodeType : uint8
 	SDK,
 	Multiply,
 	Output,
+	NullOutput,
 	Invalid
 };
 
diff --git a/Engine/Plugins/Experimental/Animation/RigMapper/Source/RigMapperEditor/Private/RigMapperGraph/SRigMapperDefinitionGraphEditor.cpp b/Engine/Plugins/Experimental/Animation/RigMapper/Source/RigMapperEditor/Private/RigMapperGraph/SRigMapperDefinitionGraphEditor.cpp
--- a/Engine/Plugins/Experimental/Animation/RigMapper/Source/RigMapperEditor/Private/RigMapperGraph/SRigMapperDefinitionGraphEditor.cpp
+++ b/Engine/Plugins/Experimental/Animation/RigMapper/Source/RigMapperEditor/Private/RigMapperGraph/SRigMapperDefinitionGraphEditor.cpp
@@ -84,7 +84,7 @@ void SRigMapperDefinitionGraphEditor::Tick(const FGeometry& AllottedGeometry, co
 	}
 }
 
-void SRigMapperDefinitionGraphEditor::SelectNodes(const TArray<FString>& Inputs, const TArray<FString>& Features, const TArray<FString>& Outputs)
+void SRigMapperDefinitionGraphEditor::SelectNodes(const TArray<FString>& Inputs, const TArray<FString>& Features, const TArray<FString>& Outputs, const TArray<FString>& NullOutputs)
 {
 	if (!bSelectingNodes)
 	{
@@ -92,7 +92,7 @@ void SRigMapperDefinitionGraphEditor::SelectNodes(const TArray<FString>& Inputs,
 		
 		GraphEditor->ClearSelectionSet();
 
-		const TArray<URigMapperDefinitionEditorGraphNode*>& SelectedNodes = GraphObj->GetNodesByName(Inputs, Features, Outputs); 
+		const TArray<URigMapperDefinitionEditorGraphNode*>& SelectedNodes = GraphObj->GetNodesByName(Inputs, Features, Outputs, NullOutputs); 
 		
 		for (URigMapperDefinitionEditorGraphNode* Node : SelectedNodes)
 		{
diff --git a/Engine/Plugins/Experimental/Animation/RigMapper/Source/RigMapperEditor/Private/RigMapperGraph/SRigMapperDefinitionGraphEditor.h b/Engine/Plugins/Experimental/Animation/RigMapper/Source/RigMapperEditor/Private/RigMapperGraph/SRigMapperDefinitionGraphEditor.h
--- a/Engine/Plugins/Experimental/Animation/RigMapper/Source/RigMapperEditor/Private/RigMapperGraph/SRigMapperDefinitionGraphEditor.h
+++ b/Engine/Plugins/Experimental/Animation/RigMapper/Source/RigMapperEditor/Private/RigMapperGraph/SRigMapperDefinitionGraphEditor.h
@@ -29,7 +29,7 @@ public:
 
 	virtual void Tick(const FGeometry& AllottedGeometry, const double InCurrentTime, const float InDeltaTime) override;
 
-	void SelectNodes(const TArray<FString>& Inputs, const TArray<FString>& Features, const TArray<FString>& Outputs);
+	void SelectNodes(const TArray<FString>& Inputs, const TArray<FString>& Features, const TArray<FString>& Outputs, const TArray<FString>& NullOutputs);
 	void RebuildGraph();
 
 private:
diff --git a/Engine/Plugins/Experimental/Animation/RigMapper/Source/RigMapperEditor/Private/RigMapperGraph/SRigMapperDefinitionStructureView.cpp b/Engine/Plugins/Experimental/Animation/RigMapper/Source/RigMapperEditor/Private/RigMapperGraph/SRigMapperDefinitionStructureView.cpp
--- a/Engine/Plugins/Experimental/Animation/RigMapper/Source/RigMapperEditor/Private/RigMapperGraph/SRigMapperDefinitionStructureView.cpp
+++ b/Engine/Plugins/Experimental/Animation/RigMapper/Source/RigMapperEditor/Private/RigMapperGraph/SRigMapperDefinitionStructureView.cpp
@@ -7,7 +7,7 @@
 #include "SlateOptMacros.h"
 #include "Widgets/Input/SSearchBox.h"
 
-const int32 SRigMapperDefinitionStructureView::NumNodeTypes = 3;
+const int32 SRigMapperDefinitionStructureView::NumNodeTypes = 4;
 const int32 SRigMapperDefinitionStructureView::NumFeatureTypes = 3;
 const FString SRigMapperDefinitionStructureView::InputsNodeName = TEXT("Inputs");
 const FString SRigMapperDefinitionStructureView::FeaturesNodeName = TEXT("Features");
@@ -15,6 +15,7 @@ const FString SRigMapperDefinitionStructureView::MultiplyNodeName = TEXT("Multip
 const FString SRigMapperDefinitionStructureView::WsNodeName = TEXT("Weighted Sums");
 const FString SRigMapperDefinitionStructureView::SdkNodeName = TEXT("SDKs");
 const FString SRigMapperDefinitionStructureView::OutputNodeName = TEXT("Outputs");
+const FString SRigMapperDefinitionStructureView::NullOutputNodeName = TEXT("Null Outputs");
 
 
 BEGIN_SLATE_FUNCTION_BUILD_OPTIMIZATION
@@ -177,10 +178,13 @@ void SRigMapperDefinitionStructureView::GenerateParentNodes()
 	
 	ParentsAndChildrenNodes.Add(ParentNodesMapping.Add(
 		ERigMapperNodeType::SDK, FeatureEntries.Add_GetRef(MakeShared<FString>(SdkNodeName))));
-	
+
 	ParentsAndChildrenNodes.Add(ParentNodesMapping.Add(
 		ERigMapperNodeType::Output, RootNodes.Add_GetRef(MakeShared<FString>(OutputNodeName))));
 
+	ParentsAndChildrenNodes.Add(ParentNodesMapping.Add(
+		ERigMapperNodeType::NullOutput, RootNodes.Add_GetRef(MakeShared<FString>(NullOutputNodeName))));
+
 	FilteredRootNodes = RootNodes;
 }
 
@@ -222,6 +226,13 @@ void SRigMapperDefinitionStructureView::GenerateChildrenNodes()
 		{
 			OutputEntries->Add(MakeShared<FString>(Output.Key));
 		}
+
+		TArray<TSharedPtr<FString>>* NullOutputEntries = GetChildrenNodes(ERigMapperNodeType::NullOutput);
+		NullOutputEntries->Reset(Definition->NullOutputs.Num());
+		for (const FString& NullOutput : Definition->NullOutputs)
+		{
+			NullOutputEntries->Add(MakeShared<FString>(NullOutput));
+		}
 	}
 }
 
@@ -230,6 +241,7 @@ void SRigMapperDefinitionStructureView::HandleTreeNodesSelectionChanged(TSharedP
 	TArray<FString> SelectedInputs;
 	TArray<FString> SelectedFeatures;
 	TArray<FString> SelectedOutputs;
+	TArray<FString> SelectedNullOutputs;
 
 	if (SelectInfo != ESelectInfo::Direct)
 	{
@@ -256,11 +268,12 @@ void SRigMapperDefinitionStructureView::HandleTreeNodesSelectionChanged(TSharedP
 		SelectedFeatures.Append(GetSelectedNodeNames(GetChildrenNodes(ERigMapperNodeType::SDK)));
 	
 		SelectedOutputs = GetSelectedNodeNames(GetChildrenNodes(ERigMapperNodeType::Output));
+		SelectedNullOutputs = GetSelectedNodeNames(GetChildrenNodes(ERigMapperNodeType::NullOutput));
 	}
 
 	if (OnSelectionChanged.IsBound())
 	{
-		OnSelectionChanged.Execute(SelectInfo, SelectedInputs, SelectedFeatures, SelectedOutputs);
+		OnSelectionChanged.Execute(SelectInfo, SelectedInputs, SelectedFeatures, SelectedOutputs, SelectedNullOutputs);
 	}
 }
 
diff --git a/Engine/Plugins/Experimental/Animation/RigMapper/Source/RigMapperEditor/Private/RigMapperGraph/SRigMapperDefinitionStructureView.h b/Engine/Plugins/Experimental/Animation/RigMapper/Source/RigMapperEditor/Private/RigMapperGraph/SRigMapperDefinitionStructureView.h
--- a/Engine/Plugins/Experimental/Animation/RigMapper/Source/RigMapperEditor/Private/RigMapperGraph/SRigMapperDefinitionStructureView.h
+++ b/Engine/Plugins/Experimental/Animation/RigMapper/Source/RigMapperEditor/Private/RigMapperGraph/SRigMapperDefinitionStructureView.h
@@ -23,7 +23,7 @@ public:
 
 	SLATE_END_ARGS()
 
-	DECLARE_DELEGATE_FourParams(FOnSelectionChanged, ESelectInfo::Type, TArray<FString>, TArray<FString>, TArray<FString>);
+	DECLARE_DELEGATE_FiveParams(FOnSelectionChanged, ESelectInfo::Type, TArray<FString>, TArray<FString>, TArray<FString>, TArray<FString>);
 	
 	/** Constructs this widget with InArgs */
 	void Construct(const FArguments& InArgs, URigMapperDefinition* InDefinition);
@@ -62,7 +62,8 @@ private:
 	static const FString WsNodeName;
 	static const FString SdkNodeName;
 	static const FString OutputNodeName;
-	
+	static const FString NullOutputNodeName;
+
 	TObjectPtr<URigMapperDefinition> Definition;
 
 	TSharedPtr<STreeView<TSharedPtr<FString>>> TreeView;
