localize a bunch of strings and also fix a minor inefficiency in AnimNode_RigMapper.cpp

diff --git a/Engine/Plugins/Experimental/Animation/RigMapper/Source/RigMapper/Private/AnimNode_RigMapper.cpp b/Engine/Plugins/Experimental/Animation/RigMapper/Source/RigMapper/Private/AnimNode_RigMapper.cpp
--- a/Engine/Plugins/Experimental/Animation/RigMapper/Source/RigMapper/Private/AnimNode_RigMapper.cpp
+++ b/Engine/Plugins/Experimental/Animation/RigMapper/Source/RigMapper/Private/AnimNode_RigMapper.cpp
@@ -120,9 +120,6 @@ void FAnimNode_RigMapper::CacheBones_AnyThread(const FAnimationCacheBonesContext
 		// todo: cache bones if required
 		// const FCompactPoseBoneIndex CompactPoseBoneIndex = RequiredBones.MakeCompactPoseIndex();
 	}
-		
-	USkeletalMesh* TargetMesh = Context.AnimInstanceProxy->GetSkelMeshComponent() ? Context.AnimInstanceProxy->GetSkelMeshComponent()->GetSkeletalMeshAsset() : nullptr;
-	InitializeRigMapping(TargetMesh);		
 }
 
 void FAnimNode_RigMapper::Evaluate_AnyThread(FPoseContext& Output)
diff --git a/Engine/Plugins/Experimental/Animation/RigMapper/Source/RigMapperEditor/Private/RigMapperDefinitionAssetTypeActions.cpp b/Engine/Plugins/Experimental/Animation/RigMapper/Source/RigMapperEditor/Private/RigMapperDefinitionAssetTypeActions.cpp
--- a/Engine/Plugins/Experimental/Animation/RigMapper/Source/RigMapperEditor/Private/RigMapperDefinitionAssetTypeActions.cpp
+++ b/Engine/Plugins/Experimental/Animation/RigMapper/Source/RigMapperEditor/Private/RigMapperDefinitionAssetTypeActions.cpp
@@ -5,6 +5,7 @@
 #include "RigMapperDefinition.h"
 #include "RigMapperDefinitionEditorToolkit.h"
 
+#define LOCTEXT_NAMESPACE "RigMapperDefinitionAssetTypeActions"
 
 UClass* FRigMapperDefinitionAssetTypeActions::GetSupportedClass() const
 {
@@ -13,7 +14,7 @@ UClass* FRigMapperDefinitionAssetTypeActions::GetSupportedClass() const
 
 FText FRigMapperDefinitionAssetTypeActions::GetName() const
 {
-	return INVTEXT("Rig Mapper Definition");
+	return LOCTEXT("FRigMapperDefinitionAssetTypeActionsName", "Rig Mapper Definition");
 }
 
 FColor FRigMapperDefinitionAssetTypeActions::GetTypeColor() const
@@ -38,4 +39,6 @@ void FRigMapperDefinitionAssetTypeActions::OpenAssetEditor(const TArray<UObject*
 			EditorToolkit->Initialize(Definition, Mode, ToolkitHost);
 		}
 	}
-}
\ No newline at end of file
+}
+
+#undef LOCTEXT_NAMESPACE
\ No newline at end of file
diff --git a/Engine/Plugins/Experimental/Animation/RigMapper/Source/RigMapperEditor/Private/RigMapperDefinitionEditorToolkit.cpp b/Engine/Plugins/Experimental/Animation/RigMapper/Source/RigMapperEditor/Private/RigMapperDefinitionEditorToolkit.cpp
--- a/Engine/Plugins/Experimental/Animation/RigMapper/Source/RigMapperEditor/Private/RigMapperDefinitionEditorToolkit.cpp
+++ b/Engine/Plugins/Experimental/Animation/RigMapper/Source/RigMapperEditor/Private/RigMapperDefinitionEditorToolkit.cpp
@@ -84,22 +84,22 @@ void FRigMapperDefinitionEditorToolkit::Initialize(URigMapperDefinition* InDefin
 
 void FRigMapperDefinitionEditorToolkit::RegisterTabSpawners(const TSharedRef<FTabManager>& InTabManager)
 {
-	WorkspaceMenuCategory = InTabManager->AddLocalWorkspaceMenuCategory(INVTEXT("Rig Mapper Definition Editor"));
+	WorkspaceMenuCategory = InTabManager->AddLocalWorkspaceMenuCategory(LOCTEXT("RigMapperDefinitionEditorTabGroup", "Rig Mapper Definition Editor"));
 
 	FAssetEditorToolkit::RegisterTabSpawners(InTabManager);
 	
 	InTabManager->RegisterTabSpawner(DefinitionEditorGraphTabId, FOnSpawnTab::CreateSP(this, &FRigMapperDefinitionEditorToolkit::SpawnGraphTab))
-		.SetDisplayName(INVTEXT("Graph"))
+		.SetDisplayName(LOCTEXT("RigMapperDefinitionEditorGraphViewName","Graph"))
 		.SetGroup(WorkspaceMenuCategory.ToSharedRef())
 		.SetIcon(FSlateIcon(FAppStyle::GetAppStyleSetName(), "GraphEditor.EventGraph_16x"));
 
 	InTabManager->RegisterTabSpawner(DefinitionEditorStructureTabId, FOnSpawnTab::CreateSP(this, &FRigMapperDefinitionEditorToolkit::SpawnStructureTab))
-		.SetDisplayName(INVTEXT("Structure"))
+		.SetDisplayName(LOCTEXT("RigMapperDefinitionEditorStructureViewName", "Structure"))
 		.SetGroup(WorkspaceMenuCategory.ToSharedRef())
 		.SetIcon(FSlateIcon(FAppStyle::GetAppStyleSetName(), "Kismet.Tabs.Palette"));
 
 	InTabManager->RegisterTabSpawner(DefinitionEditorDetailsTabId, FOnSpawnTab::CreateSP(this, &FRigMapperDefinitionEditorToolkit::SpawnDetailsTab))
-		.SetDisplayName(INVTEXT("Details"))
+		.SetDisplayName(LOCTEXT("RigMapperDefinitionEditorDetailsViewName", "Details"))
 		.SetGroup(WorkspaceMenuCategory.ToSharedRef())
 		.SetIcon(FSlateIcon(FAppStyle::GetAppStyleSetName(), "LevelEditor.Tabs.Details"));
 }
@@ -113,12 +113,12 @@ void FRigMapperDefinitionEditorToolkit::UnregisterTabSpawners(const TSharedRef<F
 
 FName FRigMapperDefinitionEditorToolkit::GetToolkitFName() const
 {
-	return "RigMapperDefinitionEditor";
+	return FName("RigMapperDefinitionEditor");
 }
 
 FText FRigMapperDefinitionEditorToolkit::GetBaseToolkitName() const
 {
-	return INVTEXT("Rig Mapper Definition Editor");
+	return LOCTEXT("RigMapperDefinitionEditorToolkitBaseToolkitName", "Rig Mapper Definition Editor");
 }
 
 FString FRigMapperDefinitionEditorToolkit::GetWorldCentricTabPrefix() const
@@ -143,7 +143,7 @@ TSharedRef<SDockTab> FRigMapperDefinitionEditorToolkit::SpawnGraphTab(const FSpa
 	
 	TSharedRef<SDockTab> DockTab = SNew(SDockTab)
 	.TabColorScale(GetTabColorScale())
-	.Label(INVTEXT("Graph"))
+	.Label(LOCTEXT("RigMapperDefinitionEditorToolkitGraphTab", "Graph"))
 	[
 		GraphEditor.ToSharedRef()
 	];
@@ -173,7 +173,7 @@ void FRigMapperDefinitionEditorToolkit::FDetailsViewCustomization::CustomizeDeta
 	TArray<TWeakObjectPtr<UObject>> Objects;
 	DetailLayout.GetObjectsBeingCustomized(Objects);
 	
-	DetailLayout.EditCategory("Animation|Rig Mapper", INVTEXT("Elements"));
+	DetailLayout.EditCategory("Animation|Rig Mapper", LOCTEXT("RigMapperDefinitionElements", "Elements"));
 	DetailLayout.HideCategory("Animation");
 	DetailLayout.HideCategory("Rig Mapper");
 }
@@ -198,7 +198,7 @@ TSharedRef<SDockTab> FRigMapperDefinitionEditorToolkit::SpawnDetailsTab(const FS
 
 	TSharedRef<SDockTab> DockTab = SNew(SDockTab)
 	.TabColorScale(GetTabColorScale())
-	.Label(INVTEXT("Details"))
+	.Label(LOCTEXT("RigMapperDefinitionEditorToolkitDetailsTab", "Details"))
 	[
 		DetailsView.ToSharedRef()
 	];
@@ -213,7 +213,7 @@ TSharedRef<SDockTab> FRigMapperDefinitionEditorToolkit::SpawnStructureTab(const
 		
 	TSharedRef<SDockTab> DockTab = SNew(SDockTab)
 	.TabColorScale(GetTabColorScale())
-	.Label(INVTEXT("Structure"))
+	.Label(LOCTEXT("RigMapperDefinitionEditorToolkitStructureTab", "Structure"))
 	[
 		StructureView.ToSharedRef()
 	];
diff --git a/Engine/Plugins/Experimental/Animation/RigMapper/Source/RigMapperEditor/Private/RigMapperEditorModule.cpp b/Engine/Plugins/Experimental/Animation/RigMapper/Source/RigMapperEditor/Private/RigMapperEditorModule.cpp
--- a/Engine/Plugins/Experimental/Animation/RigMapper/Source/RigMapperEditor/Private/RigMapperEditorModule.cpp
+++ b/Engine/Plugins/Experimental/Animation/RigMapper/Source/RigMapperEditor/Private/RigMapperEditorModule.cpp
@@ -108,7 +108,7 @@ void FRigMapperEditorModule::RegisterRigMapperDefinitionContextMenuEntries()
 	UToolMenu* Menu = UToolMenus::Get()->ExtendMenu("ContentBrowser.AssetContextMenu.RigMapperDefinition");
 
 	const FToolMenuInsert SectionInsertLocation("GetAssetActions", EToolMenuInsertType::After);
-	FToolMenuSection& Section = Menu->AddSection("RigMapperActions", INVTEXT("Rig Mapper"), SectionInsertLocation);
+	FToolMenuSection& Section = Menu->AddSection("RigMapperActions", LOCTEXT("RigMapperMenuSectionName", "Rig Mapper"), SectionInsertLocation);
 
 	Section.AddMenuEntry(
 		"RigMapperDefinition_LoadFromJson",
@@ -167,7 +167,7 @@ void FRigMapperEditorModule::RegisterRigMapperLinkedDefinitionsContextMenuEntrie
 	UToolMenu* Menu = UToolMenus::Get()->ExtendMenu("ContentBrowser.AssetContextMenu.RigMapperLinkedDefinitions");
 
 	const FToolMenuInsert SectionInsertLocation("GetAssetActions", EToolMenuInsertType::After);
-	FToolMenuSection& Section = Menu->AddSection("RigMapperActions", INVTEXT("Rig Mapper"), SectionInsertLocation);
+	FToolMenuSection& Section = Menu->AddSection("RigMapperActions", LOCTEXT("RigMapperActionsSectionName", "Rig Mapper"), SectionInsertLocation);
 
 	Section.AddMenuEntry(
 		"RigMapperLinkedDefinitions_Bake",
diff --git a/Engine/Plugins/Experimental/Animation/RigMapper/Source/RigMapperEditor/Private/RigMapperEditorSubsystem.cpp b/Engine/Plugins/Experimental/Animation/RigMapper/Source/RigMapperEditor/Private/RigMapperEditorSubsystem.cpp
--- a/Engine/Plugins/Experimental/Animation/RigMapper/Source/RigMapperEditor/Private/RigMapperEditorSubsystem.cpp
+++ b/Engine/Plugins/Experimental/Animation/RigMapper/Source/RigMapperEditor/Private/RigMapperEditorSubsystem.cpp
@@ -17,6 +17,8 @@
 #include "Interfaces//ITargetPlatformManagerModule.h"
 #include "RigMapperEditorLog.h"
 
+#define LOCTEXT_NAMESPACE "RigMapperEditorSubsystem"
+
 UAnimSequence* CreateNewAnimSequence(const FString& BaseName, const FString& Path, USkeletalMesh* TargetMesh)
 {
 	FString UniqueAssetName = BaseName;
@@ -174,7 +176,7 @@ void URigMapperEditorSubsystem::SetAnimSequenceRate(UAnimSequence* AnimSequence,
 		}
 		
 		IAnimationDataController& Controller = AnimSequence->GetController();
-		Controller.OpenBracket(INVTEXT("Set Sequence FrameRate"));
+		Controller.OpenBracket(LOCTEXT("SetSequenceFrameRateOpenBracket", "Set Sequence FrameRate"));
 		Controller.InitializeModel();
 		Controller.SetFrameRate(FrameRate, false);
 		Controller.NotifyPopulated();
@@ -258,7 +260,7 @@ bool URigMapperEditorSubsystem::AddCurveValuesToAnimSequence(UAnimSequence* Targ
 	
 	
 	IAnimationDataController& Controller = Target->GetController();
-	Controller.OpenBracket(INVTEXT("Populate New Anim"));
+	Controller.OpenBracket(LOCTEXT("PopulateNewAnimationOpenBracket", "Populate New Anim"));
 	Controller.InitializeModel();
 	Controller.RemoveAllAttributes(false);
 	Controller.RemoveAllBoneTracks(false);
@@ -1017,4 +1019,6 @@ void URigMapperEditorSubsystem::SparseBakeCurve(const FString& CurveName, int32
 			CurveValuesPerFrame[BakeFrameIndex][CurveIndex] = LastKeyedValue;
 		}
 	}
-}
\ No newline at end of file
+}
+
+#undef LOCTEXT_NAMESPACE
\ No newline at end of file
diff --git a/Engine/Plugins/Experimental/Animation/RigMapper/Source/RigMapperEditor/Private/RigMapperGraph/RigMapperDefinitionEditorGraphNode.cpp b/Engine/Plugins/Experimental/Animation/RigMapper/Source/RigMapperEditor/Private/RigMapperGraph/RigMapperDefinitionEditorGraphNode.cpp
--- a/Engine/Plugins/Experimental/Animation/RigMapper/Source/RigMapperEditor/Private/RigMapperGraph/RigMapperDefinitionEditorGraphNode.cpp
+++ b/Engine/Plugins/Experimental/Animation/RigMapper/Source/RigMapperEditor/Private/RigMapperGraph/RigMapperDefinitionEditorGraphNode.cpp
@@ -8,6 +8,8 @@
 #include "RigMapperDefinitionEditorGraph.h"
 #include "EdGraph/EdGraphPin.h"
 
+#define LOCTEXT_NAMESPACE "RigMapperDefinitionEditorGraphNode"
+
 TSharedPtr<SGraphNode> URigMapperDefinitionEditorGraphNode::NodeFactory::CreateNode(UEdGraphNode* Node) const
 {
 	if (URigMapperDefinitionEditorGraphNode* ThisNode = Cast<URigMapperDefinitionEditorGraphNode>(Node))
@@ -99,6 +101,8 @@ void URigMapperDefinitionEditorGraphNode::SetNodeType(ERigMapperNodeType InNodeT
 			{
 				const TMap<FString, double>& Inputs = static_cast<FRigMapperWsFeature*>(Feature)->Inputs;
 				const FRigMapperFeatureRange& Range = static_cast<FRigMapperWsFeature*>(Feature)->Range;
+				const FText MinText = LOCTEXT("RangeMinimum", "min");
+				const FText MaxText = LOCTEXT("RangeMaximum", "max");
 
 				for (UEdGraphPin* InPin : InputPins)
 				{
@@ -111,7 +115,7 @@ void URigMapperDefinitionEditorGraphNode::SetNodeType(ERigMapperNodeType InNodeT
 							FString RangeStr;
 							if (Range.bHasLowerBound)
 							{
-								RangeStr.Appendf(TEXT("min: %.*f"), 3, Range.LowerBound);
+								RangeStr.Appendf(TEXT("%s: %.*f"), *MinText.ToString(), 3, Range.LowerBound);
 							}
 							if (Range.bHasLowerBound && Range.bHasUpperBound)
 							{
@@ -119,7 +123,7 @@ void URigMapperDefinitionEditorGraphNode::SetNodeType(ERigMapperNodeType InNodeT
 							}
 							if (Range.bHasUpperBound)
 							{
-								RangeStr.Appendf(TEXT("max: %.*f"), 3, Range.UpperBound);
+								RangeStr.Appendf(TEXT("%s: %.*f"), *MaxText.ToString(), 3, Range.UpperBound);
 							}
 							NodeSubtitle = FText::FromString(RangeStr);
 						}	
@@ -149,3 +153,5 @@ void URigMapperDefinitionEditorGraphNode::GetRect(FVector2D& TopLeft, FVector2D&
 	BottomRight.X = NodePosX + GetDimensions().X + GetMargin().X;
 	BottomRight.Y = NodePosY + GetDimensions().Y + GetMargin().Y;
 }
+
+#undef LOCTEXT_NAMESPACE
\ No newline at end of file
diff --git a/Engine/Plugins/Experimental/Animation/RigMapper/Source/RigMapperEditor/Private/RigMapperGraph/SRigMapperDefinitionGraphEditor.cpp b/Engine/Plugins/Experimental/Animation/RigMapper/Source/RigMapperEditor/Private/RigMapperGraph/SRigMapperDefinitionGraphEditor.cpp
--- a/Engine/Plugins/Experimental/Animation/RigMapper/Source/RigMapperEditor/Private/RigMapperGraph/SRigMapperDefinitionGraphEditor.cpp
+++ b/Engine/Plugins/Experimental/Animation/RigMapper/Source/RigMapperEditor/Private/RigMapperGraph/SRigMapperDefinitionGraphEditor.cpp
@@ -28,7 +28,7 @@ SRigMapperDefinitionGraphEditor::~SRigMapperDefinitionGraphEditor()
 void SRigMapperDefinitionGraphEditor::Construct(const FArguments& InArgs, URigMapperDefinition* InDefinition)
 {
 	FGraphAppearanceInfo AppearanceInfo;
-	AppearanceInfo.CornerText = INVTEXT("Rig Mapper Definition");
+	AppearanceInfo.CornerText = LOCTEXT("GraphEditorRigMapperDefinition", "Rig Mapper Definition");
 
 	SGraphEditor::FGraphEditorEvents GraphEvents;
 	// GraphEvents.OnCreateActionMenu = SGraphEditor::FOnCreateActionMenu::CreateSP(this, &SPhysicsAssetGraph::OnCreateGraphActionMenu);
diff --git a/Engine/Plugins/Experimental/Animation/RigMapper/Source/RigMapperEditor/Private/RigMapperGraph/SRigMapperDefinitionStructureView.cpp b/Engine/Plugins/Experimental/Animation/RigMapper/Source/RigMapperEditor/Private/RigMapperGraph/SRigMapperDefinitionStructureView.cpp
--- a/Engine/Plugins/Experimental/Animation/RigMapper/Source/RigMapperEditor/Private/RigMapperGraph/SRigMapperDefinitionStructureView.cpp
+++ b/Engine/Plugins/Experimental/Animation/RigMapper/Source/RigMapperEditor/Private/RigMapperGraph/SRigMapperDefinitionStructureView.cpp
@@ -7,15 +7,17 @@
 #include "SlateOptMacros.h"
 #include "Widgets/Input/SSearchBox.h"
 
+#define LOCTEXT_NAMESPACE "RigMapperDefinitionStructureView"
+
 const int32 SRigMapperDefinitionStructureView::NumNodeTypes = 4;
 const int32 SRigMapperDefinitionStructureView::NumFeatureTypes = 3;
-const FString SRigMapperDefinitionStructureView::InputsNodeName = TEXT("Inputs");
-const FString SRigMapperDefinitionStructureView::FeaturesNodeName = TEXT("Features");
-const FString SRigMapperDefinitionStructureView::MultiplyNodeName = TEXT("Multiply");
-const FString SRigMapperDefinitionStructureView::WsNodeName = TEXT("Weighted Sums");
-const FString SRigMapperDefinitionStructureView::SdkNodeName = TEXT("SDKs");
-const FString SRigMapperDefinitionStructureView::OutputNodeName = TEXT("Outputs");
-const FString SRigMapperDefinitionStructureView::NullOutputNodeName = TEXT("Null Outputs");
+const FString SRigMapperDefinitionStructureView::InputsNodeName = LOCTEXT("RigMapperDefinitionStructureViewInputs","Inputs").ToString();
+const FString SRigMapperDefinitionStructureView::FeaturesNodeName = LOCTEXT("RigMapperDefinitionStructureViewFeatures", "Features").ToString();
+const FString SRigMapperDefinitionStructureView::MultiplyNodeName = LOCTEXT("RigMapperDefinitionStructureViewMultiplyFeatures", "Multiply").ToString();
+const FString SRigMapperDefinitionStructureView::WsNodeName = LOCTEXT("RigMapperDefinitionStructureViewWeightedSumsFeatures", "Weighted Sums").ToString();
+const FString SRigMapperDefinitionStructureView::SdkNodeName = LOCTEXT("RigMapperDefinitionStructureViewSDKsFeatures", "SDKs").ToString();
+const FString SRigMapperDefinitionStructureView::OutputNodeName = LOCTEXT("RigMapperDefinitionStructureViewOutputs", "Outputs").ToString();
+const FString SRigMapperDefinitionStructureView::NullOutputNodeName = LOCTEXT("RigMapperDefinitionStructureViewNullOutputs", "Null Outputs").ToString();
 
 
 BEGIN_SLATE_FUNCTION_BUILD_OPTIMIZATION
@@ -347,3 +349,5 @@ void SRigMapperDefinitionStructureView::OnGetTreeNodeChildren(TSharedPtr<FString
 		FilterNodes(UnfilteredChildren, Children);
 	}
 }
+
+#undef LOCTEXT_NAMESPACE
\ No newline at end of file
diff --git a/Engine/Plugins/Experimental/Animation/RigMapper/Source/RigMapperEditor/Private/RigMapperLinkedDefinitionsAssetTypeActions.cpp b/Engine/Plugins/Experimental/Animation/RigMapper/Source/RigMapperEditor/Private/RigMapperLinkedDefinitionsAssetTypeActions.cpp
--- a/Engine/Plugins/Experimental/Animation/RigMapper/Source/RigMapperEditor/Private/RigMapperLinkedDefinitionsAssetTypeActions.cpp
+++ b/Engine/Plugins/Experimental/Animation/RigMapper/Source/RigMapperEditor/Private/RigMapperLinkedDefinitionsAssetTypeActions.cpp
@@ -4,6 +4,8 @@
 
 #include "RigMapperDefinition.h"
 
+#define LOCTEXT_NAMESPACE "RigMapperLinkedDefinitionsAssetTypeActions"
+
 UClass* FRigMapperLinkedDefinitionsAssetTypeActions::GetSupportedClass() const
 {
 	return URigMapperLinkedDefinitions::StaticClass();
@@ -11,7 +13,7 @@ UClass* FRigMapperLinkedDefinitionsAssetTypeActions::GetSupportedClass() const
 
 FText FRigMapperLinkedDefinitionsAssetTypeActions::GetName() const
 {
-	return INVTEXT("Rig Mapper Linked Definitions");
+	return LOCTEXT("FRigMapperDefinitionAssetTypeActionsName", "Rig Mapper Linked Definitions");
 }
 
 FColor FRigMapperLinkedDefinitionsAssetTypeActions::GetTypeColor() const
@@ -23,3 +25,5 @@ uint32 FRigMapperLinkedDefinitionsAssetTypeActions::GetCategories()
 {
 	return EAssetTypeCategories::Animation;
 }
+
+#undef LOCTEXT_NAMESPACE
\ No newline at end of file
