fix baking bug and also disable visibility changes on selection in RigMapper editor details panel

diff --git a/Engine/Plugins/Experimental/Animation/RigMapper/Source/RigMapper/Private/RigMapperDefinition.cpp b/Engine/Plugins/Experimental/Animation/RigMapper/Source/RigMapper/Private/RigMapperDefinition.cpp
--- a/Engine/Plugins/Experimental/Animation/RigMapper/Source/RigMapper/Private/RigMapperDefinition.cpp
+++ b/Engine/Plugins/Experimental/Animation/RigMapper/Source/RigMapper/Private/RigMapperDefinition.cpp
@@ -852,7 +852,15 @@ bool FRigMapperWsFeature::BakeInput(URigMapperLinkedDefinitions* LinkedDefinitio
 
 			for (const TPair<FString, double>& SubFeatureSubInput : SubFeatureWsBakedInput->Inputs)
 			{
-				BakedWsFeature->Inputs.Add(SubFeatureSubInput.Key, SubFeatureSubInput.Value * FeatureInput.Value);
+				// check for the existance of the SubFeatureSubInput.Key in the map already as we can get 'diamond' structures
+				if (BakedWsFeature->Inputs.Contains(SubFeatureSubInput.Key))
+				{
+					BakedWsFeature->Inputs[SubFeatureSubInput.Key] += SubFeatureSubInput.Value * FeatureInput.Value;
+				}
+				else
+				{
+					BakedWsFeature->Inputs.Add(SubFeatureSubInput.Key, SubFeatureSubInput.Value * FeatureInput.Value);
+				}
 			}
 		}
 		else
diff --git a/Engine/Plugins/Experimental/Animation/RigMapper/Source/RigMapperDeveloper/Private/Tests/RigMapperTests.cpp b/Engine/Plugins/Experimental/Animation/RigMapper/Source/RigMapperDeveloper/Private/Tests/RigMapperTests.cpp
--- a/Engine/Plugins/Experimental/Animation/RigMapper/Source/RigMapperDeveloper/Private/Tests/RigMapperTests.cpp
+++ b/Engine/Plugins/Experimental/Animation/RigMapper/Source/RigMapperDeveloper/Private/Tests/RigMapperTests.cpp
@@ -1031,8 +1031,7 @@ bool FTestRigMapperCompareBakedVsUnbakedPluginDefinitions::RunTest(const FString
 		BakedDefinitionPath = TEXT("/RigMapper/Definitions/Baked/RM_CDL_FNL.RM_CDL_FNL");
 		UnbakedDefinitionPaths = { TEXT("/RigMapper/Definitions/Raw/RM_CDL_CDH.RM_CDL_CDH"),
 			TEXT("/RigMapper/Definitions/Raw/RM_CDH_FNM.RM_CDH_FNM"),
-			TEXT("/RigMapper/Definitions/Raw/RM_FNM_FNH.RM_FNM_FNH"),
-			TEXT("/RigMapper/Definitions/Raw/RM_FNH_FNL.RM_FNH_FNL")
+			TEXT("/RigMapper/Definitions/Raw/RM_FNM_FNL.RM_FNM_FNL"),
 		};
 		TestAnimSequencePath = TEXT("/Game/AutomationTestData/RigMapper/AnimSequences/FACIAL_3L_RIG_ROM.FACIAL_3L_RIG_ROM");
 	}
diff --git a/Engine/Plugins/Experimental/Animation/RigMapper/Source/RigMapperEditor/Private/RigMapperDefinitionEditorToolkit.cpp b/Engine/Plugins/Experimental/Animation/RigMapper/Source/RigMapperEditor/Private/RigMapperDefinitionEditorToolkit.cpp
--- a/Engine/Plugins/Experimental/Animation/RigMapper/Source/RigMapperEditor/Private/RigMapperDefinitionEditorToolkit.cpp
+++ b/Engine/Plugins/Experimental/Animation/RigMapper/Source/RigMapperEditor/Private/RigMapperDefinitionEditorToolkit.cpp
@@ -237,10 +237,11 @@ bool FRigMapperDefinitionEditorToolkit::HandleIsPropertyVisible(const FPropertyA
 		return true;
 	}
 	
-	if (PropertyNameToNodeTypeMapping.Contains(PropertyName))
-	{
-		return StructureView->IsNodeOrChildSelected(PropertyNameToNodeTypeMapping[PropertyName], ArrayIndex);
-	}
+	// TODO: turn off visibility logic in the Details Panel for now (MH-13360) as considered confusing to users; we may revisit this
+	//if (PropertyNameToNodeTypeMapping.Contains(PropertyName))
+	//{
+	//	return StructureView->IsNodeOrChildSelected(PropertyNameToNodeTypeMapping[PropertyName], ArrayIndex);
+	//}
 	return true;
 }
 
