Set RigMapper plugin to Experimental

diff --git a/Engine/Plugins/Experimental/Animation/RigMapper/RigMapper.uplugin b/Engine/Plugins/Experimental/Animation/RigMapper/RigMapper.uplugin
--- a/Engine/Plugins/Experimental/Animation/RigMapper/RigMapper.uplugin
+++ b/Engine/Plugins/Experimental/Animation/RigMapper/RigMapper.uplugin
@@ -13,6 +13,7 @@
 	"EnabledByDefault": false,
 	"CanContainContent": true,
 	"IsBetaVersion": false,
+	"IsExperimentalVersion": true,
 	"Installed": false,
 	"Plugins": [
 		{
@@ -32,9 +33,9 @@
 			"LoadingPhase": "PreDefault"
 		},
 		{
-			"Name" : "RigMapperDeveloper",
-			"Type" : "UncookedOnly",
-			"LoadingPhase" : "PreDefault"
+			"Name": "RigMapperDeveloper",
+			"Type": "UncookedOnly",
+			"LoadingPhase": "PreDefault"
 		}
 	]
 }
