Merge of SpeechAnimation, SpeechAnimationStreaming and RigMapper.

diff --git a/Engine/Plugins/Experimental/Animation/RigMapper/Source/RigMapper/Private/AnimNode_RigMapper.cpp b/Engine/Plugins/Experimental/Animation/RigMapper/Source/RigMapper/Private/AnimNode_RigMapper.cpp
--- a/Engine/Plugins/Experimental/Animation/RigMapper/Source/RigMapper/Private/AnimNode_RigMapper.cpp
+++ b/Engine/Plugins/Experimental/Animation/RigMapper/Source/RigMapper/Private/AnimNode_RigMapper.cpp
@@ -4,7 +4,6 @@
 
 #include "Animation/AnimInstance.h"
 #include "Components/SkeletalMeshComponent.h"
-#include "UObject/SoftObjectPath.h"
 #include "RigMapperDefinition.h"
 #include "Animation/AnimCurveUtils.h"
 #include "Engine/SkeletalMesh.h"
@@ -12,8 +11,7 @@
 #include "RigMapperLog.h"
 
 DEFINE_STAT(STAT_RigMapperEvaluateNode);
-DEFINE_STAT(STAT_RigMapperEvaluateRig);
-DEFINE_STAT(STAT_RigMapperEvaluateAllRigs);
+DEFINE_STAT(STAT_RigMapperEvaluateRigs);
 DEFINE_STAT(STAT_RigMapperInitialize);
 DEFINE_STAT(STAT_RigMapperSetCurves);
 
@@ -83,7 +81,7 @@ void FAnimNode_RigMapper::PreUpdate(const UAnimInstance* InAnimInstance)
 	{
 		for (int32 DefinitionIndex = 0; DefinitionIndex < DefinitionsToCheck->Num(); DefinitionIndex++)
 		{
-			if ((*DefinitionsToCheck)[DefinitionIndex] != LoadedDefinitions[DefinitionIndex] || !(*DefinitionsToCheck)[DefinitionIndex]->IsDefinitionValidated()) // check if individual definition has been changed or editted
+			if ((*DefinitionsToCheck)[DefinitionIndex] != LoadedDefinitions[DefinitionIndex] || !(*DefinitionsToCheck)[DefinitionIndex]->WasDefinitionValidated()) // check if individual definition has been changed or editted
 			{
 				bReInit = true;
 				break;
@@ -142,84 +140,55 @@ void FAnimNode_RigMapper::Evaluate_AnyThread(FPoseContext& Output)
 }
 
 // todo: handle bones
-void FAnimNode_RigMapper::EvaluateRigMapping(FPoseContext& Output) const
+void FAnimNode_RigMapper::EvaluateRigMapping(FPoseContext& Output)
 {
-	SCOPE_CYCLE_COUNTER(STAT_RigMapperEvaluateAllRigs);
-
-	// We need tp cache current curve values for possible Lerp using given alpha
-	TArray<double> InputValues;
-	
-	for (int32 RigMapperIndex = 0; RigMapperIndex < RigMappers.Num(); RigMapperIndex++)
+	if (RigMapperProcessor.IsValid())
 	{
-		const TSharedPtr<FacialRigMapping::FRigMapper> RigMapper = RigMappers[RigMapperIndex];
-		RigMapper->SetDirty();
-		
-		// Set first rig mapper's input values from current pose
-		if (RigMapperIndex == 0)
-		{
-			// Retrieve values from current pose
-			InputValues.AddZeroed(InputCurveMappings.Num());
-			UE::Anim::FCurveUtils::BulkGet(Output.Curve, InputCurveMappings, 
-	[&InputValues](const FRigMapperCurveMapping& InBulkElement, const float InValue)
-				{
-					InputValues[InBulkElement.CurveIndex] = InValue;
-				}
-			);
-
-			// Set first rig mapper's input values from current pose
-			for (int32 InputIndex = 0; InputIndex < InputValues.Num(); InputIndex++)
-			{
-				RigMapper->SetDirectValue(InputIndex, InputValues[InputIndex]);
-			}
-		}
-		else
-		{
-			// Other RigMappers: Link/Map prev rig mapper outputs to current inputs
-			check(IndexMaps.IsValidIndex(RigMapperIndex - 1))
-			const FRigMapperIndexMap& IndexMapping = IndexMaps[RigMapperIndex - 1];
-
-			// Get output values from previous rig mapper
-			TArray<double> PrevOutputValuesInOrder;
-			{
-				SCOPE_CYCLE_COUNTER(STAT_RigMapperEvaluateRig);
-				PrevOutputValuesInOrder = RigMappers[RigMapperIndex - 1]->GetOutputValuesInOrder();
-			}
-
-			// Set input values for current rig mapper
-			for (int32 InputIndex = 0; InputIndex < IndexMapping.Num(); InputIndex++)
+		// Retrieve inputs
+		CachedInputValues.Reset(InputCurveMappings.Num());
+		CachedInputValues.AddDefaulted(InputCurveMappings.Num());
+		UE::Anim::FCurveUtils::BulkGet(Output.Curve, InputCurveMappings, 
+	[&Inputs=CachedInputValues](const FRigMapperCurveMapping& InBulkElement, const float InValue)
 			{
-				if (const int32 OutputIndex = IndexMapping[InputIndex]; OutputIndex != INDEX_NONE)
-				{
-					RigMapper->SetDirectValue(InputIndex, PrevOutputValuesInOrder[OutputIndex]);
-				}
+				Inputs[InBulkElement.CurveIndex] = InValue;
 			}
-		}
-	}
-	
-	// Set current pose from last rig mapper's output values 
-	if (!RigMappers.IsEmpty())
-	{
-		// Retrieve last rig mapper output values
-		TArray<double> OutputValues;
+		);
+		
 		{
-			SCOPE_CYCLE_COUNTER(STAT_RigMapperEvaluateRig);
-			OutputValues = RigMappers.Last()->GetOutputValuesInOrder();
+			// Evaluate frame
+			SCOPE_CYCLE_COUNTER(STAT_RigMapperEvaluateRigs);
+			RigMapperProcessor.EvaluateFrame(RigMapperProcessor.GetInputNames(), CachedInputValues, CachedOutputValues);
 		}
 
 		{
 			SCOPE_CYCLE_COUNTER(STAT_RigMapperSetCurves);
 
-			// Set all output curves for current pose from the given output values. Possibly lerping with current values depending on Alpha.
+			// Set all output curves for current pose from the given output values.
 			UE::Anim::FCurveUtils::BulkSet(Output.Curve, OutputCurveMappings, 
-						[&InputValues, &OutputValues, LerpAlpha=Alpha](const FRigMapperOutputCurveMapping& InBulkElement)
+						[&Inputs=CachedInputValues, &Outputs=CachedOutputValues, LerpAlpha=Alpha](const FRigMapperOutputCurveMapping& InBulkElement)
 				{
-					float Value = OutputValues[InBulkElement.CurveIndex];
-							
-					if (InputValues.IsValidIndex(InBulkElement.InputCurveIndex))
+					TOptional<float> Value = Outputs[InBulkElement.CurveIndex];
+
+					// If we have a mapping for a matching input, lerp with the new output using specified alpha
+					if (Inputs.IsValidIndex(InBulkElement.InputCurveIndex))
 					{
-						Value = FMath::Lerp(InputValues[InBulkElement.InputCurveIndex], Value, LerpAlpha);
+						const TOptional<float> PrevValue = Inputs[InBulkElement.InputCurveIndex];
+
+						if (PrevValue.IsSet())
+						{
+							if (Value.IsSet())
+							{
+								Value = FMath::Lerp(PrevValue.GetValue(), Value.GetValue(), LerpAlpha);
+							}
+							else
+							{
+								// The value was not set by the rig mapper, do not lerp and keep the previous input value 
+								Value = Inputs[InBulkElement.InputCurveIndex];
+							}
+						}
 					}
-					return Value;
+					// If the value was not set by the rig mapper, and we don't have a matching input, default to 0
+					return Value.Get(0.f);
 				}
 			);
 		}
@@ -238,14 +207,12 @@ bool FAnimNode_RigMapper::InitializeRigMapping(USkeletalMesh* TargetMesh)
 	TArray<URigMapperDefinition*> DefinitionsToLoad;
 
 	// Retrieve the definitions to load and use (either the ones set on the node, or the ones overriden by the SKM asset user data)
-	bool bSkelMeshOverride = false;
 	if (TargetMesh)
 	{
 		LoadedUserData = Cast<URigMapperDefinitionUserData>(TargetMesh->GetAssetUserDataOfClass(URigMapperDefinitionUserData::StaticClass()));
 		if (LoadedUserData)
 		{
 			DefinitionsToLoad = LoadedUserData->Definitions;
-			bSkelMeshOverride = true;
 		}
 	}
 	if (DefinitionsToLoad.IsEmpty())
@@ -253,70 +220,15 @@ bool FAnimNode_RigMapper::InitializeRigMapping(USkeletalMesh* TargetMesh)
 		DefinitionsToLoad = Definitions;
 	}
 
-	if (DefinitionsToLoad.IsEmpty())
+	RigMapperProcessor = FRigMapperProcessor(DefinitionsToLoad);
+	if (!RigMapperProcessor.IsValid())
 	{
-		if (bSkelMeshOverride)
-		{
-			UE_LOG(LogRigMapper, Error, TEXT("Target SkelMesh %s has Rig Mapper Definition User Data set, but this contains no definitions."), *TargetMesh->GetName())
-		}
 		return false;
 	}
-	for (int i = 0; i < DefinitionsToLoad.Num(); i++)
-	{
-		if (!DefinitionsToLoad[i])
-		{
-			if (bSkelMeshOverride)
-			{
-				UE_LOG(LogRigMapper, Error, TEXT("Target SkelMesh %s has Rig Mapper Definition User Data set which contains an empty RigMapper definition asset at index %d; please check that all User Data RigMapper definitions have been set to valid assets."), *TargetMesh->GetName(), i)
-			}
-			else
-			{
-				UE_LOG(LogRigMapper, Error, TEXT("Rig Mapper anim node contains an empty RigMapper definition asset at index %d; please check that all RigMapper node definitions have been set to valid assets."), i)
-			}
-			return false;
-		}
-		else
-		{
-			// check if the definition is validated ... if not, warn
-			if (!DefinitionsToLoad[i]->IsDefinitionValidated())
-			{
-				if (bSkelMeshOverride)
-				{
-					UE_LOG(LogRigMapper, Warning, TEXT("Target SkelMesh %s has Rig Mapper Definition User Data set which contains an unvalidated RigMapper definition asset at index %d; please ensure that all User Data RigMapper definitions have been set to validated assets."), *TargetMesh->GetName(), i)
-				}
-				else
-				{
-					UE_LOG(LogRigMapper, Warning, TEXT("Rig Mapper anim node contains an unvalidated RigMapper definition asset at index %d; please ensure that all RigMapper node definitions have been set to validated assets."), i)
-				}
-			}
-		}
-	}
-
-	// Initialize the actual RigMappers from the definitions to load
-	RigMappers.Empty(DefinitionsToLoad.Num());
-	IndexMaps.Empty(DefinitionsToLoad.Num() - 1);
-	
-	for (const URigMapperDefinition* Definition : DefinitionsToLoad)
-	{
-		TSharedPtr<FacialRigMapping::FRigMapper> RigMapper = MakeShared<FacialRigMapping::FRigMapper>();
+	LoadedDefinitions = MoveTemp(DefinitionsToLoad);
 	
-		if (!RigMapper->LoadDefinition(Definition))
-		{
-			UE_LOG(LogRigMapper, Error, TEXT("Could not load definition %s"), *Definition->GetPathName())
-			RigMappers.Empty();
-			return false;
-		}
-
-		if (RigMappers.Num() > 0)
-		{
-			IndexMaps.Add(MakeIndexMap(RigMappers.Last(), RigMapper));
-		}
-
-		RigMappers.Add(RigMapper);
-	}
-
 	// Cache a map of curve indices to bulk get current curve values for the current pose
-	const TArray<FString>& InputNames = RigMappers[0]->GetInputNames();
+	const TArray<FString>& InputNames = RigMapperProcessor.GetInputNames();
 	InputCurveMappings.Empty();
 	InputCurveMappings.Reserve(InputNames.Num());
 	for (int32 InputIndex = 0; InputIndex < InputNames.Num(); InputIndex++)
@@ -325,7 +237,7 @@ bool FAnimNode_RigMapper::InitializeRigMapping(USkeletalMesh* TargetMesh)
 	}
 
 	// Cache a map of curve indices to bulk set the new curve values for the output pose
-	const TArray<FString>& OutputNames = RigMappers.Last()->GetOutputNames();
+	const TArray<FString>& OutputNames = RigMapperProcessor.GetOutputNames();
 	OutputCurveMappings.Empty();
 	OutputCurveMappings.Reserve(OutputNames.Num());
 	for (int32 OutputIndex = 0; OutputIndex < OutputNames.Num(); OutputIndex++)
@@ -337,22 +249,5 @@ bool FAnimNode_RigMapper::InitializeRigMapping(USkeletalMesh* TargetMesh)
 		OutputCurveMappings.Add(*OutputNames[OutputIndex], OutputIndex, InputIndex);
 	}
 	
-	LoadedDefinitions = MoveTemp(DefinitionsToLoad);
 	return true;
 }
-
-FRigMapperIndexMap FAnimNode_RigMapper::MakeIndexMap(const TSharedPtr<FacialRigMapping::FRigMapper>& OutputRig, const TSharedPtr<FacialRigMapping::FRigMapper>& InputRig)
-{
-	const TArray<FString> FromOutputs = OutputRig->GetOutputNames();
-	const TArray<FString> ToInputs = InputRig->GetInputNames();
-	
-	FRigMapperIndexMap IndexMapping;
-	
-	IndexMapping.Reserve(ToInputs.Num());
-	for (const FString& Input : ToInputs)
-	{
-		IndexMapping.Add(FromOutputs.Find(Input));
-	}
-
-	return IndexMapping;
-}
\ No newline at end of file
diff --git a/Engine/Plugins/Experimental/Animation/RigMapper/Source/RigMapper/Private/RigMapper.cpp b/Engine/Plugins/Experimental/Animation/RigMapper/Source/RigMapper/Private/RigMapper.cpp
--- a/Engine/Plugins/Experimental/Animation/RigMapper/Source/RigMapper/Private/RigMapper.cpp
+++ b/Engine/Plugins/Experimental/Animation/RigMapper/Source/RigMapper/Private/RigMapper.cpp
@@ -3,7 +3,6 @@
 #include "RigMapper.h"
 
 #include "RigMapperDefinition.h"
-#include "UObject/SoftObjectPath.h"
 #include "RigMapperLog.h"
 
 namespace FacialRigMapping
@@ -92,7 +91,7 @@ namespace FacialRigMapping
 	{
 		TMap<FString, double> OutValues;
 	
-		for (int i = 0; i < OutputNodes.Num(); i++)
+		for (int32 i = 0; i < OutputNodes.Num(); i++)
 		{
 			double OutputValue = 0.0f;
 			
@@ -110,16 +109,60 @@ namespace FacialRigMapping
 		return OutputNames;
 	}
 
-	TArray<double> FRigMapper::GetOutputValuesInOrder() const
+	void FRigMapper::GetOutputValuesInOrder(TArray<double>& OutValues) const
 	{
-		TArray<double> OutValues;
+		if (OutValues.Num() != OutputNodes.Num())
+		{
+			OutValues.Reset(OutputNodes.Num());
+			OutValues.AddZeroed(OutputNodes.Num());
+		}
 	
-		for (int i = 0; i < OutputNodes.Num(); i++)
+		for (int32 i = 0; i < OutputNodes.Num(); i++)
 		{
 			OutValues.Add(OutputNodes[i]->GetValue());
 		}
+	}
+	
+	void FRigMapper::GetOptionalOutputValuesInOrder(TArray<TOptional<double>>& OutValues) const
+	{
 
-		return OutValues;
+	
+		for (int32 i = 0; i < OutputNodes.Num(); i++)
+		{
+			double OutputValue;
+			
+			if (OutputNodes[i]->TryGetValue(OutputValue))
+			{
+				OutValues[i] = OutputValue;
+			}
+			else
+			{
+				OutValues[i].Reset();
+			}
+		}
+	}
+
+	void FRigMapper::GetOptionalFloatOutputValuesInOrder(TArray<TOptional<float>>& OutValues) const
+	{
+		if (OutValues.Num() != OutputNodes.Num())
+		{
+			OutValues.Reset(OutputNodes.Num());
+			OutValues.AddDefaulted(OutputNodes.Num());
+		}
+	
+		for (int32 i = 0; i < OutputNodes.Num(); i++)
+		{
+			double OutputValue;
+			
+			if (OutputNodes[i]->TryGetValue(OutputValue))
+			{
+				OutValues[i] = OutputValue;
+			}
+			else
+			{
+				OutValues[i].Reset();
+			}
+		}
 	}
 
 	void FRigMapper::SetDirty() const
diff --git a/Engine/Plugins/Experimental/Animation/RigMapper/Source/RigMapper/Private/RigMapperDefinition.cpp b/Engine/Plugins/Experimental/Animation/RigMapper/Source/RigMapper/Private/RigMapperDefinition.cpp
--- a/Engine/Plugins/Experimental/Animation/RigMapper/Source/RigMapper/Private/RigMapperDefinition.cpp
+++ b/Engine/Plugins/Experimental/Animation/RigMapper/Source/RigMapper/Private/RigMapperDefinition.cpp
@@ -381,8 +381,7 @@ bool URigMapperDefinition::LoadFromJsonFile(const FFilePath& JsonFilePath)
 
 void URigMapperDefinition::Empty()
 {
-	Modify();
-	bDefinitionValidated = false;
+	SetDefinitionValid(false);
 	Inputs.Empty();
 	Outputs.Empty();
 	Features.Empty();
@@ -401,34 +400,34 @@ bool URigMapperDefinition::LoadFromJsonString(const FString& JsonString)
 	}
 
 	Empty();
-	bDefinitionValidated = true;
+	SetDefinitionValid(true);
 	
 	TArray<FString> FeatureNames;
 
 	const bool bValidInputs = LoadInputsFromJsonObject(JsonObject);
 	if (!bValidInputs)
 	{
-		bDefinitionValidated = false;
+		SetDefinitionValid(false);
 	}
 	const bool bValidFeatures = LoadFeaturesFromJsonObject(JsonObject, FeatureNames);
 	if (!bValidFeatures)
 	{
-		bDefinitionValidated = false;
+		SetDefinitionValid(false);
 	}
 	const bool bValidOutputs = LoadOutputsFromJsonObject(JsonObject, FeatureNames);
 	if (!bValidOutputs)
 	{
-		bDefinitionValidated = false;
+		SetDefinitionValid(false);
 	}
 	const bool bValidNullOutputs = LoadNullOutputsFromJsonObject(JsonObject);
 	if (!bValidNullOutputs)
 	{
-		bDefinitionValidated = false;
+		SetDefinitionValid(false);
 	}
 
 	OnRigMapperDefinitionUpdated.Broadcast();
 
-	return bDefinitionValidated;
+	return bValidated;
 }
 
 bool URigMapperDefinition::LoadInputsFromJsonObject(const TSharedPtr<FJsonObject>& JsonObject)
@@ -436,7 +435,7 @@ bool URigMapperDefinition::LoadInputsFromJsonObject(const TSharedPtr<FJsonObject
 	if (!JsonObject->HasTypedField(TEXT("inputs"), EJson::Array))
 	{
 		UE_LOG(LogRigMapper, Warning, TEXT("Missing inputs field"))
-		bDefinitionValidated = false;
+		SetDefinitionValid(false);
 		return false;
 	}
 	for (const auto& InputAttr : JsonObject->GetArrayField(TEXT("inputs")))
@@ -454,9 +453,9 @@ bool URigMapperDefinition::LoadInputsFromJsonObject(const TSharedPtr<FJsonObject
 	if (Inputs.IsEmpty())
 	{
 		UE_LOG(LogRigMapper, Warning, TEXT("Not enough inputs"))
-		bDefinitionValidated = false;
+		SetDefinitionValid(false);
 	}
-	return bDefinitionValidated;
+	return bValidated;
 }
 
 bool URigMapperDefinition::LoadNullOutputsFromJsonObject(const TSharedPtr<FJsonObject>& JsonObject)
@@ -475,7 +474,7 @@ bool URigMapperDefinition::LoadNullOutputsFromJsonObject(const TSharedPtr<FJsonO
 				if (Outputs.Contains(NullOutputString))
 				{
 					UE_LOG(LogRigMapper, Warning, TEXT("Null output conflicts with existing output: %s"), *NullOutputString)
-					bDefinitionValidated = false;
+					SetDefinitionValid(false);
 				}
 				else
 				{
@@ -484,7 +483,17 @@ bool URigMapperDefinition::LoadNullOutputsFromJsonObject(const TSharedPtr<FJsonO
 			}
 		}
 	}
-	return bDefinitionValidated;
+	return WasDefinitionValidated();
+}
+
+void URigMapperDefinition::SetDefinitionValid(bool bValid)
+{
+	if (bValidated != bValid)
+	{
+		Modify();
+		bValidated = bValid;
+		UE_LOG(LogRigMapper, Log, TEXT("Definition %s is now %s"), *GetName(), WasDefinitionValidated() ? TEXT("valid") : TEXT("invalid"))
+	}
 }
 
 bool URigMapperDefinition::LoadFeaturesFromJsonObject(const TSharedPtr<FJsonObject>& JsonObject, TArray<FString>& FeatureNames)
@@ -492,7 +501,7 @@ bool URigMapperDefinition::LoadFeaturesFromJsonObject(const TSharedPtr<FJsonObje
 	if (!JsonObject->HasTypedField(TEXT("features"), EJson::Object))
 	{
 		UE_LOG(LogRigMapper, Warning, TEXT("Missing features field"))
-		bDefinitionValidated = false;
+		SetDefinitionValid(false);
 		return false;
 	}
 	for (const auto& FeatureInfo : JsonObject->GetObjectField(TEXT("features"))->Values)
@@ -500,19 +509,19 @@ bool URigMapperDefinition::LoadFeaturesFromJsonObject(const TSharedPtr<FJsonObje
 		if (Inputs.Contains(FeatureInfo.Key))
 		{
 			UE_LOG(LogRigMapper, Warning, TEXT("Feature conflicting with input of similar name: %s"), *FeatureInfo.Key)
-			bDefinitionValidated = false;
+			SetDefinitionValid(false);
 			continue;
 		}
 
 		if (!Features.AddFromJsonObject(FeatureInfo.Key, FeatureInfo.Value->AsObject()))
 		{
-			bDefinitionValidated = false;
+			SetDefinitionValid(false);
 			continue;
 		}
 
 		FeatureNames.Add(FeatureInfo.Key);
 	}
-	return bDefinitionValidated;
+	return WasDefinitionValidated();
 }
 
 bool URigMapperDefinition::LoadOutputsFromJsonObject(const TSharedPtr<FJsonObject>& JsonObject, const TArray<FString>& FeatureNames)
@@ -520,7 +529,7 @@ bool URigMapperDefinition::LoadOutputsFromJsonObject(const TSharedPtr<FJsonObjec
 	if (!JsonObject->HasTypedField(TEXT("outputs"), EJson::Object))
 	{
 		UE_LOG(LogRigMapper, Warning, TEXT("Missing outputs field"))
-		bDefinitionValidated = false;
+		SetDefinitionValid(false);
 		return false;
 	}
 	for (const auto& OutputAttr : JsonObject->GetObjectField(TEXT("outputs"))->Values)
@@ -534,7 +543,7 @@ bool URigMapperDefinition::LoadOutputsFromJsonObject(const TSharedPtr<FJsonObjec
 		if (!OutputAttr.Value.IsValid() || OutputAttr.Value->IsNull() || OutputAttr.Value->AsString().IsEmpty())
 		{
 			UE_LOG(LogRigMapper, Warning, TEXT("Invalid value for output %s"), *OutputAttr.Key)
-			bDefinitionValidated = false;
+			SetDefinitionValid(false);
 		}
 		else
 		{
@@ -542,7 +551,7 @@ bool URigMapperDefinition::LoadOutputsFromJsonObject(const TSharedPtr<FJsonObjec
 			if (!Inputs.Contains(OutputValue) && !FeatureNames.Contains(OutputValue))
 			{
 				UE_LOG(LogRigMapper, Warning, TEXT("Could not find corresponding input/feature for output %s (%s)"), *OutputAttr.Key, *OutputValue)
-				bDefinitionValidated = false;
+				SetDefinitionValid(false);
 			}
 		}
 		Outputs.Add(OutputAttr.Key, OutputValue);
@@ -550,9 +559,9 @@ bool URigMapperDefinition::LoadOutputsFromJsonObject(const TSharedPtr<FJsonObjec
 	if (Outputs.IsEmpty())
 	{
 		UE_LOG(LogRigMapper, Warning, TEXT("Not enough outputs"))
-		bDefinitionValidated = false;
+		SetDefinitionValid(false);
 	}
-	return bDefinitionValidated;
+	return WasDefinitionValidated();
 }
 
 #if WITH_EDITOR
@@ -661,14 +670,14 @@ bool URigMapperDefinition::ExportAsJsonFile(const FFilePath& JsonFilePath) const
 	return false;
 }
 
-bool URigMapperDefinition::IsDefinitionValidated() const
+bool URigMapperDefinition::WasDefinitionValidated() const
 {
-	return bDefinitionValidated;
+	return bValidated;
 }
 
-bool URigMapperDefinition::IsDefinitionValid(bool bWarn) const
+bool URigMapperDefinition::IsDefinitionValid(bool bWarn, bool bForce) const
 {
-	if (bDefinitionValidated)
+	if (!bForce && WasDefinitionValidated())
 	{
 		return true;
 	}
@@ -743,13 +752,17 @@ bool URigMapperDefinition::IsDefinitionValid(bool bWarn) const
 bool URigMapperDefinition::Validate()
 {
 	UE_LOG(LogRigMapper, Log, TEXT("Validating definition %s"), *GetName())
+
+	const bool bPrev = WasDefinitionValidated();
 	
-	bDefinitionValidated = false;
-	bDefinitionValidated = IsDefinitionValid(true);
+	SetDefinitionValid(IsDefinitionValid(true, true));
 
-	UE_LOG(LogRigMapper, Log, TEXT("Definition %s is %s"), *GetName(), bDefinitionValidated ? TEXT("valid") : TEXT("invalid"))
+	if (WasDefinitionValidated() == bPrev)
+	{
+		UE_LOG(LogRigMapper, Log, TEXT("Definition %s is still %s"), *GetName(), WasDefinitionValidated() ? TEXT("valid") : TEXT("invalid"))
+	}
 	
-	return bDefinitionValidated;
+	return WasDefinitionValidated();
 }
 
 #if WITH_EDITOR
@@ -758,7 +771,7 @@ void URigMapperDefinition::PostEditChangeProperty(FPropertyChangedEvent& Propert
 	Super::PostEditChangeProperty(PropertyChangedEvent);
 	OnRigMapperDefinitionUpdated.Broadcast();
 
-	bDefinitionValidated = false;
+	SetDefinitionValid(false);
 }
 #endif
 
diff --git a/Engine/Plugins/Experimental/Animation/RigMapper/Source/RigMapper/Private/RigMapperProcessor.cpp b/Engine/Plugins/Experimental/Animation/RigMapper/Source/RigMapper/Private/RigMapperProcessor.cpp
new file mode 100644
--- /dev/null
+++ b/Engine/Plugins/Experimental/Animation/RigMapper/Source/RigMapper/Private/RigMapperProcessor.cpp
@@ -0,0 +1,220 @@
+// Copyright Epic Games, Inc. All Rights Reserved.
+
+#include "RigMapperProcessor.h"
+#include "RigMapperLog.h"
+#include "RigMapperDefinition.h"
+
+
+FRigMapperProcessor::FRigMapperProcessor(const TArray<URigMapperDefinition*>& InDefinitions)
+{
+	LoadDefinitions(InDefinitions, RigMappers);
+	IndexCache.AddDefaulted(RigMappers.Num());
+}
+
+FRigMapperProcessor::FRigMapperProcessor(URigMapperDefinition* InDefinition)
+{
+	LoadDefinitions({ InDefinition }, RigMappers);
+	IndexCache.AddDefaulted(RigMappers.Num());
+}
+
+FRigMapperProcessor::FRigMapperProcessor(const TArray<TSharedPtr<FacialRigMapping::FRigMapper>>& InRigMappers): RigMappers(InRigMappers)
+{
+	IndexCache.AddDefaulted(RigMappers.Num());	
+}
+
+FRigMapperProcessor::FRigMapperProcessor(TSharedPtr<FacialRigMapping::FRigMapper> InRigMapper): RigMappers( { InRigMapper } )
+{
+	IndexCache.AddDefaulted();
+}
+
+bool FRigMapperProcessor::LoadDefinitions(const TArray<URigMapperDefinition*>& Definitions, TArray<TSharedPtr<FacialRigMapping::FRigMapper>>& RigMappers)
+{
+	if (Definitions.IsEmpty())
+	{
+		UE_LOG(LogRigMapper, Error, TEXT("Not enough definitions to remap"))
+		return false;
+	}
+	
+	RigMappers.Empty(Definitions.Num());
+	
+	for (int32 DefinitionIndex = 0; DefinitionIndex < Definitions.Num(); DefinitionIndex++)
+	{
+		if (!Definitions[DefinitionIndex])
+		{
+			UE_LOG(LogRigMapper, Error, TEXT("Invalid definition at index: %d"), DefinitionIndex)
+			RigMappers.Empty();
+			return false;
+		}
+
+		// check if the definition is validated ... if not, warn
+		if (!Definitions[DefinitionIndex]->WasDefinitionValidated())
+		{
+			UE_LOG(LogRigMapper, Warning, TEXT("RigMapper definition asset at index %d (%s) was not validated please ensure that all RigMapper node definitions have been set to validated assets."), DefinitionIndex, *Definitions[DefinitionIndex]->GetPathName())
+		}
+		
+		TSharedPtr<FacialRigMapping::FRigMapper> RigMapper = MakeShared<FacialRigMapping::FRigMapper>();
+		
+		if (!RigMapper->LoadDefinition(Definitions[DefinitionIndex]))
+		{
+			UE_LOG(LogRigMapper, Error, TEXT("Could not load definition at index %d (%s)"), DefinitionIndex, *Definitions[DefinitionIndex]->GetPathName())
+			RigMappers.Empty();
+			return false;
+		}
+		
+		RigMappers.Add(RigMapper);
+	}
+	
+	// todo: check definitions chain 
+	
+	return true;
+}
+
+bool FRigMapperProcessor::IsValid() const
+{
+	if (RigMappers.IsEmpty())
+	{
+		return false;
+	}
+	for (int32 RigIndex = 0; RigIndex < RigMappers.Num(); RigIndex++)
+	{
+		if (!RigMappers[RigIndex])
+		{
+			return false;
+		}
+	}
+	return true;
+}
+
+const TArray<FString>& FRigMapperProcessor::GetInputNames() const
+{
+	return RigMappers[0]->GetInputNames();
+}
+
+const TArray<FString>& FRigMapperProcessor::GetOutputNames() const
+{
+	return RigMappers.Last()->GetOutputNames();
+}
+
+bool FRigMapperProcessor::EvaluateFrames(const TArray<FString>& CurveNames, const TArray<FPoseValues>& InFrameValues, TArray<FPoseValues>& OutFrameValues)
+{
+	if (!IsValid())
+	{
+		return false;
+	}
+	bool bValid = true;
+	
+	for (int32 RigIndex = 0; RigIndex < RigMappers.Num(); RigIndex++)
+	{
+		OutFrameValues.Reset(InFrameValues.Num());
+		OutFrameValues.AddDefaulted(InFrameValues.Num());
+	
+		for (int32 FrameIndex = 0; FrameIndex < InFrameValues.Num(); FrameIndex++)
+		{
+			if (RigIndex == 0)
+			{
+				bValid &= EvaluateFrame_Internal(CurveNames, InFrameValues[FrameIndex], OutFrameValues[FrameIndex], RigIndex);
+			}
+			else
+			{
+				bValid &= EvaluateFrame_Internal(RigMappers[RigIndex - 1]->GetOutputNames(), OutFrameValues[FrameIndex], OutFrameValues[FrameIndex], RigIndex);
+			}
+		}
+	}
+	return bValid;
+}
+
+bool FRigMapperProcessor::EvaluateFrames(const TArray<FString>& CurveNames, const TArray<FPoseValues>& InFrameValues, TArray<FString>& OutCurveNames,
+	TArray<FPoseValues>& OutFrameValues)
+{
+	if (!IsValid() || !EvaluateFrames(CurveNames, InFrameValues, OutFrameValues))
+	{
+		return false;
+	}
+	const TArray<FString>& NewOutputNames = RigMappers.Last()->GetOutputNames();
+	OutCurveNames.Reset(NewOutputNames.Num());
+	OutCurveNames.Append(NewOutputNames);
+	return true;
+}
+
+bool FRigMapperProcessor::EvaluateFrames_Interp(const TArray<FString>& CurveNames, const TArray<FPoseValues>& InFrameValues, TArray<FPoseValues>& OutFrameValues, const TArray<FFrameTime>& FrameTimes)
+{
+	// todo
+	if (!IsValid() || !EvaluateFrames(CurveNames, InFrameValues, OutFrameValues))
+	{
+		return false;
+	}
+	return true;
+}
+
+bool FRigMapperProcessor::EvaluateFrame(const TArray<FString>& CurveNames, const FPoseValues& InCurveValues, FPoseValues& OutCurveValues)
+{
+	if (!IsValid())
+	{
+		return false;
+	}
+	bool bValid = true;
+	
+	for (int32 RigIndex = 0; RigIndex < RigMappers.Num(); RigIndex++)
+	{
+		if (RigIndex == 0)
+		{
+			bValid &= EvaluateFrame_Internal(CurveNames, InCurveValues, OutCurveValues, RigIndex);
+		}
+		else
+		{
+			bValid &= EvaluateFrame_Internal(RigMappers[RigIndex - 1]->GetOutputNames(), OutCurveValues, OutCurveValues, RigIndex);
+		}
+	}
+
+	return true;
+}
+
+bool FRigMapperProcessor::EvaluateFrame(const TArray<FString>& CurveNames, const FPoseValues& InCurveValues, TArray<FString>& OutCurveNames,
+	FPoseValues& OutCurveValues)
+{
+	if (!IsValid() || !EvaluateFrame(CurveNames, InCurveValues, OutCurveValues))
+	{
+		return false;
+	}
+	const TArray<FString>& NewOutputNames = RigMappers.Last()->GetOutputNames();
+	OutCurveNames.Reset(NewOutputNames.Num());
+	OutCurveNames.Append(NewOutputNames);
+	return true;
+}
+
+
+bool FRigMapperProcessor::EvaluateFrame_Internal(const TArray<FString>& CurveNames, const FPoseValues& InCurveValues, FPoseValues& OutCurveValues, const int32 RigIndex)
+{
+	if (!RigMappers.IsValidIndex(RigIndex) || !RigMappers[RigIndex].IsValid() || CurveNames.Num() != InCurveValues.Num())
+	{
+		return false;
+	}
+	if (IndexCache[RigIndex].Max() != CurveNames.Num())
+	{
+		IndexCache[RigIndex].Reset(CurveNames.Num());
+	}
+		
+	RigMappers[RigIndex]->SetDirty();
+
+	// Set inputs
+	const TArray<FString>& InputNames = RigMappers[RigIndex]->GetInputNames();
+		
+	for (int32 CurveIndex = 0; CurveIndex < CurveNames.Num(); CurveIndex++)
+	{
+		// Cache input indices
+		if (!IndexCache[RigIndex].IsValidIndex(CurveIndex))
+		{
+			IndexCache[RigIndex].Add(InputNames.Find(CurveNames[CurveIndex]));
+		}
+
+		// Set input from curve
+		const int32 InputIndex = IndexCache[RigIndex][CurveIndex];
+		if (InputNames.IsValidIndex(InputIndex) && InCurveValues[CurveIndex].IsSet())
+		{
+			RigMappers[RigIndex]->SetDirectValue(InputIndex, InCurveValues[CurveIndex].GetValue());
+		}
+	}
+	
+	RigMappers[RigIndex]->GetOptionalFloatOutputValuesInOrder(OutCurveValues);
+	return true;
+}
diff --git a/Engine/Plugins/Experimental/Animation/RigMapper/Source/RigMapper/Public/AnimNode_RigMapper.h b/Engine/Plugins/Experimental/Animation/RigMapper/Source/RigMapper/Public/AnimNode_RigMapper.h
--- a/Engine/Plugins/Experimental/Animation/RigMapper/Source/RigMapper/Public/AnimNode_RigMapper.h
+++ b/Engine/Plugins/Experimental/Animation/RigMapper/Source/RigMapper/Public/AnimNode_RigMapper.h
@@ -5,7 +5,7 @@
 #include "Animation/AnimNodeBase.h"
 
 #include "CoreMinimal.h"
-#include "RigMapper.h"
+#include "RigMapperProcessor.h"
 
 #include "AnimNode_RigMapper.generated.h"
 
@@ -19,8 +19,7 @@ class USkeletalMesh;
 DECLARE_STATS_GROUP(TEXT("RigMapper"), STATGROUP_RigMapper, STATCAT_Advanced);
 
 DECLARE_CYCLE_STAT_EXTERN(TEXT("Evaluate Rig Mapper Node"), STAT_RigMapperEvaluateNode, STATGROUP_RigMapper, RIGMAPPER_API);
-DECLARE_CYCLE_STAT_EXTERN(TEXT("Evaluate Rig Mapping"), STAT_RigMapperEvaluateAllRigs, STATGROUP_RigMapper, RIGMAPPER_API);
-DECLARE_CYCLE_STAT_EXTERN(TEXT("Evaluate Rig Mapper Definition"), STAT_RigMapperEvaluateRig, STATGROUP_RigMapper, RIGMAPPER_API);
+DECLARE_CYCLE_STAT_EXTERN(TEXT("Evaluate Rig Mapping"), STAT_RigMapperEvaluateRigs, STATGROUP_RigMapper, RIGMAPPER_API);
 DECLARE_CYCLE_STAT_EXTERN(TEXT("Initialize Rig Mapping"), STAT_RigMapperInitialize, STATGROUP_RigMapper, RIGMAPPER_API);
 DECLARE_CYCLE_STAT_EXTERN(TEXT("Rig Mapper Set Pose Curves"), STAT_RigMapperSetCurves, STATGROUP_RigMapper, RIGMAPPER_API);
 
@@ -54,7 +53,7 @@ public:
 	bool InitializeRigMapping(USkeletalMesh* TargetMesh=nullptr);
 
 	// Evaluate the new curve values using the initialized rig mappers and sets the new output pose 
-	void EvaluateRigMapping(FPoseContext& Output) const;
+	void EvaluateRigMapping(FPoseContext& Output);
 
 public:
 	UPROPERTY(EditAnywhere, BlueprintReadWrite, Category = Links)
@@ -74,24 +73,23 @@ public:
 
 	UPROPERTY(EditAnywhere, BlueprintReadWrite, Category = Alpha, meta = (PinShownByDefault))
 	float Alpha = 1.0f;
-
-private:
-	// Builds an In/Out index mapping between 2 rig mapper definitions chained together
-	static FRigMapperIndexMap MakeIndexMap(const TSharedPtr<FacialRigMapping::FRigMapper>& OutputRig, const TSharedPtr<FacialRigMapping::FRigMapper>& InputRig);
 	
 private:
-	// The inititalized rig mappers to evaluate the remapped pose
-	TArray<TSharedPtr<FacialRigMapping::FRigMapper>> RigMappers;
-
 	// The definitions that we have loaded. Cached to check against changes and reinit if need be
 	TArray<TObjectPtr<URigMapperDefinition>> LoadedDefinitions;
 
 	// The asset user data currently used to override definitions, if any was set on the skeletal mesh
 	TObjectPtr<URigMapperDefinitionUserData> LoadedUserData;
 
-	// Mapping between rig mappers inputs and outputs in order
-	TArray<FRigMapperIndexMap> IndexMaps;
-
+	// The processor to evaluate the rig mapping
+	FRigMapperProcessor RigMapperProcessor;
+	
+	// The cached input values passed to the rig mapper processor to avoid reallocations
+	FRigMapperProcessor::FPoseValues CachedInputValues;
+	
+	// The cached output values passed to the rig mapper processor to avoid reallocations
+	FRigMapperProcessor::FPoseValues CachedOutputValues;
+	
 	// Base curve mapping for bulk get/set of the linked pose curves
 	struct FRigMapperCurveMapping
 	{
diff --git a/Engine/Plugins/Experimental/Animation/RigMapper/Source/RigMapper/Public/RigMapper.h b/Engine/Plugins/Experimental/Animation/RigMapper/Source/RigMapper/Public/RigMapper.h
--- a/Engine/Plugins/Experimental/Animation/RigMapper/Source/RigMapper/Public/RigMapper.h
+++ b/Engine/Plugins/Experimental/Animation/RigMapper/Source/RigMapper/Public/RigMapper.h
@@ -28,7 +28,9 @@ namespace FacialRigMapping
 
 		TMap<FString,double> GetOutputValues(bool bSkipUnset = false) const;
 		const TArray<FString>& GetOutputNames() const;
-		TArray<double> GetOutputValuesInOrder() const;
+		void GetOutputValuesInOrder(TArray<double>& OutValues) const;
+		void GetOptionalOutputValuesInOrder(TArray<TOptional<double>>& OutValues) const;
+		void GetOptionalFloatOutputValuesInOrder(TArray<TOptional<float>>& OutValues) const;
 		
 		void SetDirty() const;
 
diff --git a/Engine/Plugins/Experimental/Animation/RigMapper/Source/RigMapper/Public/RigMapperDefinition.h b/Engine/Plugins/Experimental/Animation/RigMapper/Source/RigMapper/Public/RigMapperDefinition.h
--- a/Engine/Plugins/Experimental/Animation/RigMapper/Source/RigMapper/Public/RigMapperDefinition.h
+++ b/Engine/Plugins/Experimental/Animation/RigMapper/Source/RigMapper/Public/RigMapperDefinition.h
@@ -271,13 +271,13 @@ public:
 	UPROPERTY(EditAnywhere, BlueprintReadWrite, Category="Animation|Rig Mapper")
 	TArray<FString> NullOutputs;
 
+	UPROPERTY(VisibleAnywhere, BlueprintReadOnly, Category="Animation|Rig Mapper")
+	bool bValidated = false;
+	
 private:
 	UPROPERTY()
 	FString JsonSource;
 
-	UPROPERTY()
-	bool bDefinitionValidated = false;
-
 public:
 	UFUNCTION(BlueprintCallable, Category="Animation|Rig Mapper")
 	bool LoadFromJsonFile(const FFilePath& JsonFilePath);
@@ -295,13 +295,13 @@ public:
 	void Empty();
 
 	UFUNCTION(BlueprintCallable, Category = "Animation|Rig Mapper")
-	bool IsDefinitionValid(bool bWarn = false) const;
+	bool IsDefinitionValid(bool bWarn = false, bool bForce = false) const;
 
 	UFUNCTION(BlueprintCallable, Category="Animation|Rig Mapper")
 	bool Validate();
 
 	UFUNCTION(BlueprintCallable, Category = "Animation|Rig Mapper")
-	bool IsDefinitionValidated() const;
+	bool WasDefinitionValidated() const;
 
 	
 private:
@@ -309,6 +309,7 @@ private:
 	bool LoadFeaturesFromJsonObject(const TSharedPtr<FJsonObject>& JsonObject, TArray<FString>& FeatureNames);
 	bool LoadOutputsFromJsonObject(const TSharedPtr<FJsonObject>& JsonObject, const TArray<FString>& FeatureNames);
 	bool LoadNullOutputsFromJsonObject(const TSharedPtr<FJsonObject>& JsonObject);
+	void SetDefinitionValid(bool bValid);
 
 #if WITH_EDITOR
 public:
diff --git a/Engine/Plugins/Experimental/Animation/RigMapper/Source/RigMapper/Public/RigMapperProcessor.h b/Engine/Plugins/Experimental/Animation/RigMapper/Source/RigMapper/Public/RigMapperProcessor.h
new file mode 100644
--- /dev/null
+++ b/Engine/Plugins/Experimental/Animation/RigMapper/Source/RigMapper/Public/RigMapperProcessor.h
@@ -0,0 +1,70 @@
+// Copyright Epic Games, Inc. All Rights Reserved.
+
+#pragma once
+
+#include "CoreMinimal.h"
+#include "RigMapper.h"
+#include "Misc/FrameTime.h"
+#include "RigMapperProcessor.generated.h"
+
+/**
+ * A helper struct to gain performance when doing some batch remapping with the Rig Mapper pipeline
+ * Useful when remapping frames every tick, or numerous frames at once, or when using chained definitions...
+ * EvaluateFrame(s) methods work best if given the same set of CurveNames at every call 
+ * Also provides a static utility to load rig mappers from definitions
+ */
+USTRUCT()
+struct RIGMAPPER_API FRigMapperProcessor
+{
+	GENERATED_BODY()
+		
+	using FPose = TMap<FString, float>;
+	using FPoseValues = TArray<TOptional<float>>;
+	using FSparseBakedCurves = TMap<FFrameTime, FPoseValues>;
+
+	FRigMapperProcessor() = default;
+	explicit FRigMapperProcessor(const TArray<URigMapperDefinition*>& InDefinitions);
+	explicit FRigMapperProcessor(URigMapperDefinition* InDefinition);
+	explicit FRigMapperProcessor(TSharedPtr<FacialRigMapping::FRigMapper> InRigMapper);
+	explicit FRigMapperProcessor(const TArray<TSharedPtr<FacialRigMapping::FRigMapper>>& InRigMappers);
+
+	/** Loads a set of rig mappers from a set of definition assets */
+	static bool LoadDefinitions(const TArray<URigMapperDefinition*>& Definitions, TArray<TSharedPtr<FacialRigMapping::FRigMapper>>& RigMappers);
+
+	/** @returns whether we have at least one rig mapper initialized, and all rig mappers are valid */
+	bool IsValid() const;
+
+	/** @returns the input names expected by the first rig mapper */
+	const TArray<FString>& GetInputNames() const;
+
+	/** @returns the output names given by the last rig mapper */
+	const TArray<FString>& GetOutputNames() const;
+
+	/** Evaluates a set of frames all at once - retrieve and/or cache curve names using @GetOutputNames for performance */
+	bool EvaluateFrames(const TArray<FString>& CurveNames, const TArray<FPoseValues>& InFrameValues, TArray<FPoseValues>& OutFrameValues);
+
+	/** Evaluates a set of frames all at once */
+	bool EvaluateFrames(const TArray<FString>& CurveNames, const TArray<FPoseValues>& InFrameValues, TArray<FString>& OutCurveNames, TArray<FPoseValues>& OutFrameValues);
+
+	/** Evaluates a set of frames all at once, interpolating between them using the given frame times */
+	bool EvaluateFrames_Interp(const TArray<FString>& CurveNames, const TArray<FPoseValues>& InFrameValues, TArray<FPoseValues>& OutFrameValues, const TArray<FFrameTime>& FrameTimes);
+
+	/** Evaluates a single frame - retrieve and/or cache curve names using @GetOutputNames for performance */
+	bool EvaluateFrame(const TArray<FString>& CurveNames, const FPoseValues& InCurveValues, FPoseValues& OutCurveValues);
+
+	/** Evaluates a single frame */
+	bool EvaluateFrame(const TArray<FString>& CurveNames, const FPoseValues& InCurveValues, TArray<FString>& OutCurveNames, FPoseValues& OutCurveValues);
+
+private:
+	/** Evaluates a single frame for a single rig mapper */
+	bool EvaluateFrame_Internal(const TArray<FString>& CurveNames, const FPoseValues& InCurveValues, FPoseValues& OutCurveValues, int32 RigIndex);
+	
+private:
+	using FIndexCache = TArray<int32>;
+
+	/** The rig mappers loaded from definitions */
+	TArray<TSharedPtr<FacialRigMapping::FRigMapper>> RigMappers;
+
+	/** An index map for each rig mapper's input - if the input curve names change between two calls, it will be rebuilt */
+	TArray<FIndexCache> IndexCache;
+};
diff --git a/Engine/Plugins/Experimental/Animation/RigMapper/Source/RigMapperDeveloper/Private/Tests/RigMapperTests.cpp b/Engine/Plugins/Experimental/Animation/RigMapper/Source/RigMapperDeveloper/Private/Tests/RigMapperTests.cpp
--- a/Engine/Plugins/Experimental/Animation/RigMapper/Source/RigMapperDeveloper/Private/Tests/RigMapperTests.cpp
+++ b/Engine/Plugins/Experimental/Animation/RigMapper/Source/RigMapperDeveloper/Private/Tests/RigMapperTests.cpp
@@ -392,7 +392,7 @@ void EvaluateRigMapping(const TMap<FString, double>& InInputs, const TArray<FRig
 			// Get output values from previous rig mapper
 			TArray<double> PrevOutputValuesInOrder;
 			{
-				PrevOutputValuesInOrder = InOutRigMappers[i - 1]->GetOutputValuesInOrder();
+				InOutRigMappers[i - 1]->GetOutputValuesInOrder(PrevOutputValuesInOrder);
 			}
 
 			// Set input values for current rig mapper
diff --git a/Engine/Plugins/Experimental/Animation/RigMapper/Source/RigMapperEditor/Private/AnimGraphNode_RigMapper.cpp b/Engine/Plugins/Experimental/Animation/RigMapper/Source/RigMapperEditor/Private/AnimGraphNode_RigMapper.cpp
--- a/Engine/Plugins/Experimental/Animation/RigMapper/Source/RigMapperEditor/Private/AnimGraphNode_RigMapper.cpp
+++ b/Engine/Plugins/Experimental/Animation/RigMapper/Source/RigMapperEditor/Private/AnimGraphNode_RigMapper.cpp
@@ -68,7 +68,15 @@ void UAnimGraphNode_RigMapper::PostEditChangeProperty(FPropertyChangedEvent& Pro
 	const FName PropertyName = PropertyChangedEvent.GetPropertyName();
 	if (PropertyName == GET_MEMBER_NAME_CHECKED(FAnimNode_RigMapper, Definitions))
 	{
-		USkeletalMesh* TargetMesh = GetAnimBlueprint() ? (GetAnimBlueprint()->GetPreviewMesh() ? GetAnimBlueprint()->GetPreviewMesh() : GetAnimBlueprint()->TargetSkeleton->GetPreviewMesh()) : nullptr;
+		USkeletalMesh* TargetMesh = nullptr;
+		if (const UAnimBlueprint* AnimBlueprint = GetAnimBlueprint())
+		{
+			TargetMesh = AnimBlueprint->GetPreviewMesh();
+			if (!TargetMesh && AnimBlueprint->TargetSkeleton)
+			{
+				TargetMesh = AnimBlueprint->TargetSkeleton->GetPreviewMesh();
+			}
+		}
 		Node.InitializeRigMapping(TargetMesh);
 		ReconstructNode();
 	}
diff --git a/Engine/Plugins/Experimental/Animation/RigMapper/Source/RigMapperEditor/Private/RigMapperEditorSubsystem.cpp b/Engine/Plugins/Experimental/Animation/RigMapper/Source/RigMapperEditor/Private/RigMapperEditorSubsystem.cpp
--- a/Engine/Plugins/Experimental/Animation/RigMapper/Source/RigMapperEditor/Private/RigMapperEditorSubsystem.cpp
+++ b/Engine/Plugins/Experimental/Animation/RigMapper/Source/RigMapperEditor/Private/RigMapperEditorSubsystem.cpp
@@ -16,6 +16,7 @@
 #include "Units/Execution/RigUnit_PrepareForExecution.h"
 #include "Interfaces//ITargetPlatformManagerModule.h"
 #include "RigMapperEditorLog.h"
+#include "RigMapperProcessor.h"
 
 #define LOCTEXT_NAMESPACE "RigMapperEditorSubsystem"
 
@@ -186,29 +187,19 @@ void URigMapperEditorSubsystem::SetAnimSequenceRate(UAnimSequence* AnimSequence,
 	}
 }
 
-bool URigMapperEditorSubsystem::ConvertCurveValuesToCsv(TArray<FString>& CurveNames, const TArray<FFrameTime>& FrameTimes, TArray<FPoseValues>& CurveValuesPerFrame, const FFilePath& OutputFile, const TArray<URigMapperDefinition*>& Definitions, bool bOutputIntermediateCsvFiles)
+bool URigMapperEditorSubsystem::ConvertCurveValuesToCsv(TArray<FString>& CurveNames, const TArray<FFrameTime>& FrameTimes, FFrameValues& CurveValuesPerFrame, const FFilePath& OutputFile, const TArray<URigMapperDefinition*>& Definitions, bool bOutputIntermediateCsvFiles)
 {
-	TArray<TSharedPtr<FacialRigMapping::FRigMapper>> RigMappers;
-	if (!InitializeRigMapping(Definitions, RigMappers))
+	FRigMapperProcessor RigMapperProcessor(Definitions);
+	if (!RigMapperProcessor.IsValid())
 	{
-		// todo: warn
 		return false;
 	}
 
-	const FString BaseFilePath = FPaths::GetBaseFilename(OutputFile.FilePath, false);;
+	// todo: bOutputIntermediateCsvFiles
+	// const FString BaseFilePath = FPaths::GetBaseFilename(OutputFile.FilePath, false);;
 	
-	for (int32 RigIndex = 0; RigIndex < RigMappers.Num(); RigIndex++)
-	{
-		if (bOutputIntermediateCsvFiles)
-		{
-			FFilePath CurrentOutputFile;
-			CurrentOutputFile.FilePath = FString::Printf(TEXT("%s_%d.%s"), *BaseFilePath, RigIndex, *FPaths::GetExtension(OutputFile.FilePath));
-			
-			WriteCurveValuesToCsv(CurrentOutputFile, CurveNames, FrameTimes, CurveValuesPerFrame);
-		}
-		
-		EvaluateRigMapping(CurveNames, FrameTimes, CurveValuesPerFrame, RigMappers[RigIndex]);
-	}
+	RigMapperProcessor.EvaluateFrames(CurveNames, CurveValuesPerFrame, CurveNames, CurveValuesPerFrame);
+
 	WriteCurveValuesToCsv(OutputFile, CurveNames, FrameTimes, CurveValuesPerFrame);
 	
 	return true;
@@ -218,7 +209,7 @@ bool URigMapperEditorSubsystem::ConvertCsv(const FFilePath& InputFile, const FFi
 {
 	TArray<FString> CurveNames;
 	TArray<FFrameTime> FrameTimes;
-		TArray<FPoseValues> CurveValuesPerFrame;
+		FFrameValues CurveValuesPerFrame;
 	
 	if (!LoadCurveValuesFromCsv(InputFile, CurveNames, FrameTimes, CurveValuesPerFrame))
 	{
@@ -228,37 +219,48 @@ bool URigMapperEditorSubsystem::ConvertCsv(const FFilePath& InputFile, const FFi
 	return ConvertCurveValuesToCsv(CurveNames, FrameTimes, CurveValuesPerFrame, OutputFile, Definitions, bOutputIntermediateCsvFiles);
 }
 
-bool URigMapperEditorSubsystem::ConvertCurveValuesToAnimSequence(TArray<FString>& CurveNames, const TArray<FFrameTime>& FrameTimes, TArray<FPoseValues>& CurveValuesPerFrame, UAnimSequence* Target, const TArray<URigMapperDefinition*>& Definitions)
+bool URigMapperEditorSubsystem::ConvertCurveValuesToAnimSequence(TArray<FString>& CurveNames, const TArray<FFrameTime>& FrameTimes, FFrameValues& CurveValuesPerFrame, UAnimSequence* Target, const TArray<URigMapperDefinition*>& Definitions)
 {
-	TArray<TSharedPtr<FacialRigMapping::FRigMapper>> RigMappers;
-	if (!InitializeRigMapping(Definitions, RigMappers))
+	FRigMapperProcessor RigMapperProcessor(Definitions);
+	if (!RigMapperProcessor.IsValid())
 	{
-		// todo: warn
+		UE_LOG(LogRigMapperEditor, Error, TEXT("Rig Mapper Processor is invalid"))
 		return false;
 	}
-	
-	for (int32 RigIndex = 0; RigIndex < RigMappers.Num(); RigIndex++)
+	FFrameValues ConvertedCurveValuesPerFrame;
+	if (!RigMapperProcessor.EvaluateFrames(CurveNames, CurveValuesPerFrame, CurveNames, ConvertedCurveValuesPerFrame))
 	{
-		EvaluateRigMapping(CurveNames, FrameTimes, CurveValuesPerFrame, RigMappers[RigIndex]);
+		UE_LOG(LogRigMapperEditor, Error, TEXT("Rig Mapper Processor failed to evaluate frames"))
+		return false;
 	}
-	
-	return AddCurveValuesToAnimSequence(Target, CurveNames, FrameTimes, CurveValuesPerFrame);
+	return AddCurveValuesToAnimSequence(Target, CurveNames, FrameTimes, ConvertedCurveValuesPerFrame);
 }
 
-bool URigMapperEditorSubsystem::AddCurveValuesToAnimSequence(UAnimSequence* Target, const TArray<FString>& CurveNames, const TArray<FFrameTime>& FrameTimes, const TArray<FPoseValues>& CurveValuesPerFrame)
+bool URigMapperEditorSubsystem::AddCurveValuesToAnimSequence(UAnimSequence* Target, const TArray<FString>& CurveNames, const TArray<FFrameTime>& FrameTimes, const FFrameValues& CurveValuesPerFrame)
 {	
-	if (CurveNames.IsEmpty() || FrameTimes.IsEmpty() || CurveValuesPerFrame.Num() != FrameTimes.Num())
+	if (CurveNames.IsEmpty())
 	{
+		UE_LOG(LogRigMapperEditor, Error, TEXT("Failed to add empty curves list"))
 		return false;
 	}
-
+	if (FrameTimes.IsEmpty())
+	{
+		UE_LOG(LogRigMapperEditor, Error, TEXT("Failed to add curves with no frame times"))
+		return false;
+	}
+	if (CurveValuesPerFrame.Num() != FrameTimes.Num())
+	{
+		UE_LOG(LogRigMapperEditor, Error, TEXT("Failed to add curves: number of frames does not match expected number (%d vs %d)"), CurveValuesPerFrame.Num(), FrameTimes.Num())
+		return false;
+	}
+	
 	const FFrameRate FrameRate = GetAnimSequenceRate(Target);
 	if (!FrameRate.IsValid())
 	{
+		UE_LOG(LogRigMapperEditor, Error, TEXT("Invalid rate %f"), FrameRate.AsDecimal())
 		return false;
 	}
 	
-	
 	IAnimationDataController& Controller = Target->GetController();
 	Controller.OpenBracket(LOCTEXT("PopulateNewAnimationOpenBracket", "Populate New Anim"));
 	Controller.InitializeModel();
@@ -284,9 +286,12 @@ bool URigMapperEditorSubsystem::AddCurveValuesToAnimSequence(UAnimSequence* Targ
 			for (int32 FrameIndex = 0; FrameIndex < FrameTimes.Num(); FrameIndex++)
 			{
 				const float Time = FrameRate.AsSeconds(FrameTimes[FrameIndex]) - TimeOffset;
-				const float Value = CurveValuesPerFrame[FrameIndex][CurveIndex];
+				const TOptional<float> Value = CurveValuesPerFrame[FrameIndex][CurveIndex];
 
-				Controller.SetCurveKey(CurveId, FRichCurveKey(Time, Value), false);
+				if (Value.IsSet())
+				{
+					Controller.SetCurveKey(CurveId, FRichCurveKey(Time, Value.GetValue()), false);
+				}
 			}
 		}
 	}
@@ -299,9 +304,9 @@ bool URigMapperEditorSubsystem::AddCurveValuesToAnimSequence(UAnimSequence* Targ
 	return true;
 }
 
-bool URigMapperEditorSubsystem::LoadCurveValuesFromAnimSequence(const UAnimSequence* Source, TArray<FString>& CurveNames, TArray<FFrameTime>& FrameTimes, TArray<FPoseValues>& CurveValuesPerFrame)
+bool URigMapperEditorSubsystem::LoadCurveValuesFromAnimSequence(const UAnimSequence* Source, TArray<FString>& CurveNames, TArray<FFrameTime>& FrameTimes, FFrameValues& CurveValuesPerFrame)
 {
-	TArray<FPose> Poses;
+	FFramePoses Poses;
 	
 	const TArray<FFloatCurve>& FloatCurves = Source->GetDataModel()->GetFloatCurves();
 
@@ -342,12 +347,14 @@ bool URigMapperEditorSubsystem::ConvertAnimSequence(const UAnimSequence* Source,
 {
 	TArray<FString> CurveNames;
 	TArray<FFrameTime> FrameTimes;
-	TArray<FPoseValues> CurveValuesPerFrame;
-	
+	FFrameValues CurveValuesPerFrame;
+
+	UE_LOG(LogRigMapperEditor, Log, TEXT("Loading curves from anim sequence"))
 	if (!LoadCurveValuesFromAnimSequence(Source, CurveNames, FrameTimes, CurveValuesPerFrame))
 	{
 		return false;
 	}
+	UE_LOG(LogRigMapperEditor, Log, TEXT("Converting curves to anim sequence: %d named curves, %d frame times, %d frame curves"), CurveNames.Num(), FrameTimes.Num(), CurveValuesPerFrame.Num())
 	
 	return ConvertCurveValuesToAnimSequence(CurveNames, FrameTimes, CurveValuesPerFrame, Target, Definitions);
 }
@@ -372,7 +379,7 @@ bool URigMapperEditorSubsystem::ConvertCsvToAnimSequence(const FFilePath& InputF
 {
 	TArray<FString> CurveNames;
 	TArray<FFrameTime> FrameTimes;
-	TArray<FPoseValues> CurveValuesPerFrame;
+	FFrameValues CurveValuesPerFrame;
 	
 	if (!LoadCurveValuesFromCsv(InputFile, CurveNames, FrameTimes, CurveValuesPerFrame))
 	{
@@ -400,7 +407,7 @@ bool URigMapperEditorSubsystem::ConvertAnimSequenceToCsv(const UAnimSequence* So
 {
 	TArray<FString> CurveNames;
 	TArray<FFrameTime> FrameTimes;
-	TArray<FPoseValues> CurveValuesPerFrame;
+	FFrameValues CurveValuesPerFrame;
 	
 	if (!LoadCurveValuesFromAnimSequence(Source, CurveNames, FrameTimes, CurveValuesPerFrame))
 	{
@@ -414,7 +421,7 @@ bool URigMapperEditorSubsystem::ConvertAnimSequenceToControlRigSection(const UAn
 {
 	TArray<FString> CurveNames;
 	TArray<FFrameTime> FrameTimes;
-	TArray<FPoseValues> CurveValuesPerFrame;
+	FFrameValues CurveValuesPerFrame;
 	
 	if (!LoadCurveValuesFromAnimSequence(Source, CurveNames, FrameTimes, CurveValuesPerFrame))
 	{
@@ -430,7 +437,7 @@ UMovieSceneControlRigParameterSection* URigMapperEditorSubsystem::ConvertAnimSeq
 	
 	TArray<FString> CurveNames;
 	TArray<FFrameTime> FrameTimes;
-	TArray<FPoseValues> CurveValuesPerFrame;
+	FFrameValues CurveValuesPerFrame;
 	
 	if (!LoadCurveValuesFromAnimSequence(Source, CurveNames, FrameTimes, CurveValuesPerFrame))
 	{
@@ -444,19 +451,14 @@ UMovieSceneControlRigParameterSection* URigMapperEditorSubsystem::ConvertAnimSeq
 	return nullptr;
 }
 
-bool URigMapperEditorSubsystem::ConvertCurveValuesToControlRigSection(TArray<FString>& CurveNames, const TArray<FFrameTime>& FrameTimes, TArray<FPoseValues>& CurveValuesPerFrame, UMovieSceneControlRigParameterSection* Target, const TArray<URigMapperDefinition*>& Definitions)
+bool URigMapperEditorSubsystem::ConvertCurveValuesToControlRigSection(TArray<FString>& CurveNames, const TArray<FFrameTime>& FrameTimes, FFrameValues& CurveValuesPerFrame, UMovieSceneControlRigParameterSection* Target, const TArray<URigMapperDefinition*>& Definitions)
 {
-	TArray<TSharedPtr<FacialRigMapping::FRigMapper>> RigMappers;
-	if (!InitializeRigMapping(Definitions, RigMappers))
+	FRigMapperProcessor RigMapperProcessor(Definitions);
+	if (!RigMapperProcessor.IsValid())
 	{
-		// todo: warn
 		return false;
 	}
-	
-	for (int32 RigIndex = 0; RigIndex < RigMappers.Num(); RigIndex++)
-	{
-		EvaluateRigMapping(CurveNames, FrameTimes, CurveValuesPerFrame, RigMappers[RigIndex]);
-	}
+	RigMapperProcessor.EvaluateFrames(CurveNames, CurveValuesPerFrame, CurveNames, CurveValuesPerFrame);
 
 	return AddCurveValuesToControlRigSection(Target, CurveNames, FrameTimes, CurveValuesPerFrame);
 }
@@ -491,7 +493,7 @@ int32 GetChannelIndexFromCurveName(const FString& CurveName, const TArrayView<co
 	return ChannelIndex;
 }
 
-bool URigMapperEditorSubsystem::AddCurveValuesToControlRigSection(UMovieSceneControlRigParameterSection* Target, const TArray<FString>& CurveNames, const TArray<FFrameTime>& FrameTimes, const TArray<FPoseValues>& CurveValuesPerFrame)
+bool URigMapperEditorSubsystem::AddCurveValuesToControlRigSection(UMovieSceneControlRigParameterSection* Target, const TArray<FString>& CurveNames, const TArray<FFrameTime>& FrameTimes, const FFrameValues& CurveValuesPerFrame)
 {
 	const UMovieScene* MovieScene = Target->GetTypedOuter<UMovieScene>();
 
@@ -541,8 +543,13 @@ bool URigMapperEditorSubsystem::AddCurveValuesToControlRigSection(UMovieSceneCon
 				
 				for (int32 FrameIndex = 0; FrameIndex < FrameTimes.Num(); FrameIndex++)
 				{
-					const FFrameNumber FrameNumber = MovieScene->GetTickResolution().AsFrameNumber(MovieScene->GetDisplayRate().AsSeconds(FrameTimes[FrameIndex]));
-					Channel->AddCubicKey(FrameNumber, CurveValuesPerFrame[FrameIndex][CurveIndex]);
+					const TOptional<float> CurveValue = CurveValuesPerFrame[FrameIndex][CurveIndex];
+
+					if (CurveValue.IsSet())
+					{
+						const FFrameNumber FrameNumber = MovieScene->GetTickResolution().AsFrameNumber(MovieScene->GetDisplayRate().AsSeconds(FrameTimes[FrameIndex]));
+						Channel->AddCubicKey(FrameNumber, CurveValuesPerFrame[FrameIndex][CurveIndex].GetValue());
+					}
 				}
 			}
 		}
@@ -591,9 +598,9 @@ FString ChannelNameToCurveName(const FString& ChannelName, const TArray<FString>
 	return ""; // todo: warn
 }
 
-bool URigMapperEditorSubsystem::LoadCurveValuesFromControlRigSection(const UMovieSceneControlRigParameterSection* Source, TArray<FString>& CurveNames, TArray<FFrameTime>& FrameTimes, TArray<FPoseValues>& CurveValuesPerFrame, const TArray<FString>& InputNames)
+bool URigMapperEditorSubsystem::LoadCurveValuesFromControlRigSection(const UMovieSceneControlRigParameterSection* Source, TArray<FString>& CurveNames, TArray<FFrameTime>& FrameTimes, FFrameValues& CurveValuesPerFrame, const TArray<FString>& InputNames)
 {
-	TArray<FPose> Poses;
+	FFramePoses Poses;
 	
 	const UMovieScene* MovieScene = Source->GetTypedOuter<UMovieScene>();
 	check(MovieScene);
@@ -646,7 +653,7 @@ bool URigMapperEditorSubsystem::ConvertControlRigSection(const UMovieSceneContro
 {
 	TArray<FString> CurveNames;
 	TArray<FFrameTime> FrameTimes;
-	TArray<FPoseValues> CurveValuesPerFrame;
+	FFrameValues CurveValuesPerFrame;
 
 	if (Definitions.IsEmpty())
 	{
@@ -675,7 +682,7 @@ bool URigMapperEditorSubsystem::ConvertCsvToControlRigSection(const FFilePath& I
 {
 	TArray<FString> CurveNames;
 	TArray<FFrameTime> FrameTimes;
-	TArray<FPoseValues> CurveValuesPerFrame;
+	FFrameValues CurveValuesPerFrame;
 	
 	if (!LoadCurveValuesFromCsv(InputFile, CurveNames, FrameTimes, CurveValuesPerFrame))
 	{
@@ -700,7 +707,7 @@ bool URigMapperEditorSubsystem::ConvertControlRigSectionToCsv(const UMovieSceneC
 {
 	TArray<FString> CurveNames;
 	TArray<FFrameTime> FrameTimes;
-	TArray<FPoseValues> CurveValuesPerFrame;
+	FFrameValues CurveValuesPerFrame;
 
 	if (Definitions.IsEmpty())
 	{
@@ -718,7 +725,7 @@ bool URigMapperEditorSubsystem::ConvertControlRigSectionToAnimSequence(const UMo
 {
 	TArray<FString> CurveNames;
 	TArray<FFrameTime> FrameTimes;
-	TArray<FPoseValues> CurveValuesPerFrame;
+	FFrameValues CurveValuesPerFrame;
 
 	if (Definitions.IsEmpty())
 	{
@@ -732,12 +739,12 @@ bool URigMapperEditorSubsystem::ConvertControlRigSectionToAnimSequence(const UMo
 	return ConvertCurveValuesToAnimSequence(CurveNames, FrameTimes, CurveValuesPerFrame, Target, Definitions);
 }
 
-bool URigMapperEditorSubsystem::ConvertControlRigSectionToAnimSequenceNew(const UMovieSceneControlRigParameterSection* Source, USkeletalMesh* TargetMesh, const TArray<URigMapperDefinition*>& Definitions, const FDirectoryPath& NewAssetPath, const FName NewAssetName)
+UAnimSequence* URigMapperEditorSubsystem::ConvertControlRigSectionToAnimSequenceNew(const UMovieSceneControlRigParameterSection* Source, USkeletalMesh* TargetMesh, const TArray<URigMapperDefinition*>& Definitions, const FDirectoryPath& NewAssetPath, const FName NewAssetName)
 {
 	ULevelSequence* Sequence = GetSequenceFromSection(Source);
 	if (!Sequence || !Sequence->GetMovieScene())
 	{
-		return false;
+		return nullptr;
 	}
 	
 	UAnimSequence* NewSequence = CreateNewAnimSequence(NewAssetName.ToString(), NewAssetPath.Path, TargetMesh);
@@ -745,28 +752,28 @@ bool URigMapperEditorSubsystem::ConvertControlRigSectionToAnimSequenceNew(const
 	
 	TArray<FString> CurveNames;
 	TArray<FFrameTime> FrameTimes;
-	TArray<FPoseValues> CurveValuesPerFrame;
+	FFrameValues CurveValuesPerFrame;
 
 	if (Definitions.IsEmpty())
 	{
-		return false;
+		return nullptr;
 	}
 	if (!LoadCurveValuesFromControlRigSection(Source, CurveNames, FrameTimes, CurveValuesPerFrame, Definitions[0]->Inputs))
 	{
-		return false;
+		return nullptr;
 	}
 
 	if (!ConvertCurveValuesToAnimSequence(CurveNames, FrameTimes, CurveValuesPerFrame, NewSequence, Definitions))
 	{
-		return false;
+		return nullptr;
 	}
 	NewSequence->PostLoad();
-	return true;
+	return NewSequence;
 }
 
-bool URigMapperEditorSubsystem::LoadCurveValuesFromCsv(const FFilePath& InputFile, TArray<FString>& CurveNames, TArray<FFrameTime>& FrameTimes, TArray<FPoseValues>& CurveValuesPerFrame)
+bool URigMapperEditorSubsystem::LoadCurveValuesFromCsv(const FFilePath& InputFile, TArray<FString>& CurveNames, TArray<FFrameTime>& FrameTimes, FFrameValues& CurveValuesPerFrame)
 {
-	TArray<FPose> Poses;
+	FFramePoses Poses;
 	
 	TArray<FString> CsvLines;
 	
@@ -827,83 +834,7 @@ bool URigMapperEditorSubsystem::LoadCurveValuesFromCsv(const FFilePath& InputFil
 	return BakeSparseKeys(Poses, CurveNames, FrameTimes, CurveValuesPerFrame);
 }
 
-bool URigMapperEditorSubsystem::InitializeRigMapping(const TArray<URigMapperDefinition*>& Definitions, TArray<TSharedPtr<FacialRigMapping::FRigMapper>>& RigMappers)
-{
-	if (Definitions.IsEmpty())
-	{
-		UE_LOG(LogRigMapperEditor, Error, TEXT("Not enough definitions to remap"))
-		return false;
-	}
-	
-	RigMappers.Empty(Definitions.Num());
-	
-	for (int i = 0; i < Definitions.Num(); i++)
-	{
-		if (!Definitions[i])
-		{
-			UE_LOG(LogRigMapperEditor, Error, TEXT("Invalid definition at index: %d"), i)
-			RigMappers.Empty();
-			return false;
-		}
-		TSharedPtr<FacialRigMapping::FRigMapper> RigMapper = MakeShared<FacialRigMapping::FRigMapper>();
-		
-		if (!RigMapper->LoadDefinition(Definitions[i]))
-		{
-			UE_LOG(LogRigMapperEditor, Error, TEXT("Could not load definition at index %d"), i)
-			RigMappers.Empty();
-			return false;
-		}
-		
-		RigMappers.Add(RigMapper);
-	}
-
-	// todo: check definitions chain 
-	
-	return true;
-}
-
-void URigMapperEditorSubsystem::EvaluateRigMapping(TArray<FString>& CurveNames, const TArray<FFrameTime>& FrameTimes, TArray<FPoseValues>& CurveValuesPerFrame, const TSharedPtr<FacialRigMapping::FRigMapper>& RigMapper)
-{
-	TArray<FString> NewCurveNames;
-	
-	for (int32 FrameIndex = 0;  FrameIndex < FrameTimes.Num(); FrameIndex++)
-	{
-		RigMapper->SetDirty();
-		
-		const TArray<FString>& InputNames = RigMapper->GetInputNames();
-		
-		for (int32 CurveIndex = 0; CurveIndex < CurveNames.Num(); CurveIndex++)
-		{
-			if (const int32 InputIndex = InputNames.Find(CurveNames[CurveIndex]); InputIndex != INDEX_NONE)
-			{
-				RigMapper->SetDirectValue(InputIndex, CurveValuesPerFrame[FrameIndex][CurveIndex]);
-			}
-			else
-			{
-				// todo: log / interp -> missing input?
-				UE_LOG(LogRigMapperEditor, Log, TEXT("Missing key, interp not avaialble yet"))
-			}
-		}
-
-		TMap<FString, double> Outputs = RigMapper->GetOutputValues(true);
-
-		if (NewCurveNames.IsEmpty())
-		{
-			NewCurveNames.Reserve(Outputs.Num());
-		}
-		CurveValuesPerFrame[FrameIndex].Empty(Outputs.Num());
-		CurveValuesPerFrame[FrameIndex].AddDefaulted(Outputs.Num());
-		
-		for (const TPair<FString, double>& Output : Outputs)
-		{
-			const int32 CurveIndex = NewCurveNames.AddUnique(Output.Key);
-			CurveValuesPerFrame[FrameIndex][CurveIndex] = Output.Value;;
-		}			
-	}
-	CurveNames = NewCurveNames;
-}
-
-bool URigMapperEditorSubsystem::WriteCurveValuesToCsv(const FFilePath& OutputFile, const TArray<FString>& CurveNames, const TArray<FFrameTime>& FrameTimes, const TArray<FPoseValues>& CurveValuesPerFrame)
+bool URigMapperEditorSubsystem::WriteCurveValuesToCsv(const FFilePath& OutputFile, const TArray<FString>& CurveNames, const TArray<FFrameTime>& FrameTimes, const FFrameValues& CurveValuesPerFrame)
 {
 	TArray<FString> CsvLines;
 	CsvLines.Add(TEXT("curve_name, frame_number, value"));
@@ -912,20 +843,34 @@ bool URigMapperEditorSubsystem::WriteCurveValuesToCsv(const FFilePath& OutputFil
 	{
 		for (int32 CurveIndex = 0; CurveIndex < CurveNames.Num(); CurveIndex++)
 		{
-			CsvLines.Add(FString::Printf(TEXT("%s, %d, %f"), *CurveNames[CurveIndex], FrameTimes[FrameIndex].RoundToFrame().Value, CurveValuesPerFrame[FrameIndex][CurveIndex]));
+			if (CurveValuesPerFrame[FrameIndex][CurveIndex].IsSet())
+			{
+				CsvLines.Add(FString::Printf(TEXT("%s, %d, %f"), *CurveNames[CurveIndex], FrameTimes[FrameIndex].RoundToFrame().Value, CurveValuesPerFrame[FrameIndex][CurveIndex].GetValue()));
+			}
 		}
 	}
 	
 	return FFileHelper::SaveStringArrayToFile(CsvLines, *OutputFile.FilePath);
 }
 
-bool URigMapperEditorSubsystem::BakeSparseKeys(const TArray<FPose>& Poses, const TArray<FString>& CurveNames, TArray<FFrameTime>& FrameTimes, TArray<FPoseValues>& CurveValuesPerFrame)
+bool URigMapperEditorSubsystem::BakeSparseKeys(const FFramePoses& Poses, const TArray<FString>& CurveNames, TArray<FFrameTime>& FrameTimes, FFrameValues& CurveValuesPerFrame)
 {
-	if (Poses.IsEmpty() || CurveNames.IsEmpty() || Poses.Num() != FrameTimes.Num())
+	if (Poses.IsEmpty())
 	{
+		UE_LOG(LogRigMapperEditor, Error, TEXT("Failed to bake sparse keys: not enough poses"))
 		return false;
 	}
-
+	if (CurveNames.IsEmpty())
+	{
+		UE_LOG(LogRigMapperEditor, Error, TEXT("Failed to bake sparse keys: not enough curves"))
+		return false;
+	}
+	if (Poses.Num() != FrameTimes.Num())
+	{
+		UE_LOG(LogRigMapperEditor, Error, TEXT("Number of poses does not match number of frame times (%d vs %d)"), Poses.Num(), FrameTimes.Num())
+		return false;
+	}
+	
 	// This extra work is needed to make sure frame times and thus resulting curve values are in order
 	TArray<FFrameTime> FrameTimesInOrder = FrameTimes;
 	FrameTimesInOrder.Sort();
@@ -968,12 +913,27 @@ bool URigMapperEditorSubsystem::BakeSparseKeys(const TArray<FPose>& Poses, const
 
 	FrameTimes = FrameTimesInOrder;
 	
-	return CurveValuesPerFrame.Num() == FrameTimes.Num() && CurveValuesPerFrame[0].Num() == CurveNames.Num() && CurveValuesPerFrame.Last().Num() == CurveNames.Num();
+	if (CurveValuesPerFrame.Num() != FrameTimes.Num())
+	{
+		UE_LOG(LogRigMapperEditor, Error, TEXT("Number of baked frames does not match number of frame times (%d vs %d)"), CurveValuesPerFrame.Num(), FrameTimes.Num())
+		return false;
+	}
+	if (CurveValuesPerFrame[0].Num() != CurveNames.Num())
+	{
+		UE_LOG(LogRigMapperEditor, Error, TEXT("Number of curves in first frame does not match expected number of curves (%d vs %d)"), CurveValuesPerFrame[0].Num(), CurveNames.Num())
+		return false;
+	}
+	if (CurveValuesPerFrame.Last().Num() != CurveNames.Num())
+	{
+		UE_LOG(LogRigMapperEditor, Error, TEXT("Number of curves in last frame does not match expected number of curves (%d vs %d)"), CurveValuesPerFrame.Last().Num(), CurveNames.Num())
+		return false;
+	}
+	return true;
 }
 
 void URigMapperEditorSubsystem::SparseBakeCurve(const FString& CurveName, int32 CurveIndex, float CurveValue,
-	TArray<FPoseValues>& CurveValuesPerFrame, const TArray<FFrameTime>& FrameTimesInOrder, int32 FrameIndex,
-	const TArray<FPose>& Poses, const TArray<FFrameTime>& FrameTimes, const int32 ActualFrameIndex)
+	FFrameValues& CurveValuesPerFrame, const TArray<FFrameTime>& FrameTimesInOrder, int32 FrameIndex,
+	const FFramePoses& Poses, const TArray<FFrameTime>& FrameTimes, const int32 ActualFrameIndex)
 {
 	// We haven't seen this control so far. So keep it const on all previous frames as well as the current one
 	for (int32 BakeFrameIndex = 0; BakeFrameIndex <= FrameIndex; BakeFrameIndex += 1)
@@ -995,7 +955,7 @@ void URigMapperEditorSubsystem::SparseBakeCurve(const FString& CurveName, int32
 		{
 			// Key this frame
 			CurveValuesPerFrame[BakeFrameIndex][CurveIndex] = Poses[ActualBakeFrameIndex][CurveName];
-			LastKeyedValue = CurveValuesPerFrame[BakeFrameIndex][CurveIndex];
+			LastKeyedValue = CurveValuesPerFrame[BakeFrameIndex][CurveIndex].Get(0);
 			
 			// If we have not seen this control on the previous frame, lerp all keys with missed
 			if (BakeFrameIndex > LastKeyedFrameIndex + 1)
diff --git a/Engine/Plugins/Experimental/Animation/RigMapper/Source/RigMapperEditor/Public/RigMapperEditorSubsystem.h b/Engine/Plugins/Experimental/Animation/RigMapper/Source/RigMapperEditor/Public/RigMapperEditorSubsystem.h
--- a/Engine/Plugins/Experimental/Animation/RigMapper/Source/RigMapperEditor/Public/RigMapperEditorSubsystem.h
+++ b/Engine/Plugins/Experimental/Animation/RigMapper/Source/RigMapperEditor/Public/RigMapperEditorSubsystem.h
@@ -7,8 +7,8 @@
 #include "EditorSubsystem.h"
 #include "Kismet/BlueprintFunctionLibrary.h"
 
-#include "RigMapper.h"
 #include "RigMapperDefinition.h"
+#include "RigMapperProcessor.h"
 
 #include "RigMapperEditorSubsystem.generated.h"
 
@@ -26,12 +26,6 @@ UCLASS()
 class RIGMAPPEREDITOR_API URigMapperEditorSubsystem : public UEditorSubsystem
 {
 	GENERATED_BODY()
-
-	using FPose = TMap<FString, float>;
-	
-	using FPoseValues = TArray<float>;
-
-	using FSparseBakedCurves = TMap<FFrameTime, FPoseValues>;
 	
 public:
 	/**
@@ -139,7 +133,7 @@ public:
 	 * Convert a Control Rig section to a new AnimSequence asset, remapping it using the given definitions
 	 */
 	UFUNCTION(BlueprintCallable, Category = "Editor Scripting|Animation|Rig Mapper")
-	static bool ConvertControlRigSectionToAnimSequenceNew(const UMovieSceneControlRigParameterSection* Source, USkeletalMesh* TargetMesh, const TArray<URigMapperDefinition*>& Definitions, const FDirectoryPath& NewAssetPath, const FName NewAssetName);
+	static UAnimSequence* ConvertControlRigSectionToAnimSequenceNew(const UMovieSceneControlRigParameterSection* Source, USkeletalMesh* TargetMesh, const TArray<URigMapperDefinition*>& Definitions, const FDirectoryPath& NewAssetPath, const FName NewAssetName);
 	
 	UFUNCTION(BlueprintCallable, Category = "Editor Scripting|Animation|Rig Mapper")
 	static ULevelSequence* GetSequenceFromSection(const UMovieSceneControlRigParameterSection* Section);
@@ -154,29 +148,28 @@ public:
 	static void SetAnimSequenceRate(UAnimSequence* AnimSequence, FFrameRate FrameRate, bool bSetImportProperties);
 	
 private:
-	static bool ConvertCurveValuesToCsv(TArray<FString>& CurveNames, const TArray<FFrameTime>& FrameTimes, TArray<FPoseValues>& CurveValuesPerFrame, const FFilePath& OutputFile, const TArray<URigMapperDefinition*>& Definitions, bool bOutputIntermediateCsvFiles);
+	using FFrameValues = TArray<FRigMapperProcessor::FPoseValues>;
+	using FFramePoses = TArray<FRigMapperProcessor::FPose>;
+	
+	static bool ConvertCurveValuesToCsv(TArray<FString>& CurveNames, const TArray<FFrameTime>& FrameTimes, FFrameValues& CurveValuesPerFrame, const FFilePath& OutputFile, const TArray<URigMapperDefinition*>& Definitions, bool bOutputIntermediateCsvFiles);
 
-	static bool WriteCurveValuesToCsv(const FFilePath& OutputFile, const TArray<FString>& CurveNames, const TArray<FFrameTime>& FrameTimes, const TArray<FPoseValues>& CurveValuesPerFrame);
+	static bool WriteCurveValuesToCsv(const FFilePath& OutputFile, const TArray<FString>& CurveNames, const TArray<FFrameTime>& FrameTimes, const FFrameValues& CurveValuesPerFrame);
 	
-	static bool LoadCurveValuesFromCsv(const FFilePath& InputFile, TArray<FString>& CurveNames, TArray<FFrameTime>& FrameTimes, TArray<FPoseValues>& CurveValuesPerFrame);
+	static bool LoadCurveValuesFromCsv(const FFilePath& InputFile, TArray<FString>& CurveNames, TArray<FFrameTime>& FrameTimes, FFrameValues& CurveValuesPerFrame);
 	
-	static bool ConvertCurveValuesToAnimSequence(TArray<FString>& CurveNames, const TArray<FFrameTime>& FrameTimes, TArray<FPoseValues>& CurveValuesPerFrame, UAnimSequence* Target, const TArray<URigMapperDefinition*>& Definitions);
+	static bool ConvertCurveValuesToAnimSequence(TArray<FString>& CurveNames, const TArray<FFrameTime>& FrameTimes, FFrameValues& CurveValuesPerFrame, UAnimSequence* Target, const TArray<URigMapperDefinition*>& Definitions);
 
-	static bool AddCurveValuesToAnimSequence(UAnimSequence* Target, const TArray<FString>& CurveNames, const TArray<FFrameTime>& FrameTimes, const TArray<FPoseValues>& CurveValuesPerFrame);
+	static bool AddCurveValuesToAnimSequence(UAnimSequence* Target, const TArray<FString>& CurveNames, const TArray<FFrameTime>& FrameTimes, const FFrameValues& CurveValuesPerFrame);
 
-	static bool LoadCurveValuesFromAnimSequence(const UAnimSequence* Source, TArray<FString>& CurveNames, TArray<FFrameTime>& FrameTimes, TArray<FPoseValues>& CurveValuesPerFrame);
-	
-	static bool ConvertCurveValuesToControlRigSection(TArray<FString>& CurveNames, const TArray<FFrameTime>& FrameTimes, TArray<FPoseValues>& CurveValuesPerFrame, UMovieSceneControlRigParameterSection* Target, const TArray<URigMapperDefinition*>& Definitions);
+	static bool LoadCurveValuesFromAnimSequence(const UAnimSequence* Source, TArray<FString>& CurveNames, TArray<FFrameTime>& FrameTimes, FFrameValues& CurveValuesPerFrame);
 	
-	static bool AddCurveValuesToControlRigSection(UMovieSceneControlRigParameterSection* Target, const TArray<FString>& CurveNames, const TArray<FFrameTime>& FrameTimes, const TArray<FPoseValues>& CurveValuesPerFrame);
-
-	static bool LoadCurveValuesFromControlRigSection(const UMovieSceneControlRigParameterSection* Source, TArray<FString>& CurveNames, TArray<FFrameTime>& FrameTimes, TArray<FPoseValues>& CurveValuesPerFrame, const TArray<FString>& InputNames);
+	static bool ConvertCurveValuesToControlRigSection(TArray<FString>& CurveNames, const TArray<FFrameTime>& FrameTimes, FFrameValues& CurveValuesPerFrame, UMovieSceneControlRigParameterSection* Target, const TArray<URigMapperDefinition*>& Definitions);
 	
-	static bool InitializeRigMapping(const TArray<URigMapperDefinition*>& Definitions, TArray<TSharedPtr<FacialRigMapping::FRigMapper>>& RigMappers);
+	static bool AddCurveValuesToControlRigSection(UMovieSceneControlRigParameterSection* Target, const TArray<FString>& CurveNames, const TArray<FFrameTime>& FrameTimes, const FFrameValues& CurveValuesPerFrame);
 
-	static void EvaluateRigMapping(TArray<FString>& CurveNames, const TArray<FFrameTime>& FrameTimes, TArray<FPoseValues>& CurveValuesPerFrame, const TSharedPtr<FacialRigMapping::FRigMapper>& RigMapper);
+	static bool LoadCurveValuesFromControlRigSection(const UMovieSceneControlRigParameterSection* Source, TArray<FString>& CurveNames, TArray<FFrameTime>& FrameTimes, FFrameValues& CurveValuesPerFrame, const TArray<FString>& InputNames);
 
-	static bool BakeSparseKeys(const TArray<FPose>& Poses, const TArray<FString>& CurveNames, TArray<FFrameTime>& FrameTimes, TArray<FPoseValues>& CurveValuesPerFrame);
+	static bool BakeSparseKeys(const FFramePoses& Poses, const TArray<FString>& CurveNames, TArray<FFrameTime>& FrameTimes, FFrameValues& CurveValuesPerFrame);
 
-	static void SparseBakeCurve(const FString& CurveName, int32 CurveIndex, float CurveValue, TArray<FPoseValues>& CurveValuesPerFrame, const TArray<FFrameTime>& FrameTimesInOrder, int32 FrameIndex, const TArray<FPose>& Poses, const TArray<FFrameTime>& FrameTimes, const int32 ActualFrameIndex);
+	static void SparseBakeCurve(const FString& CurveName, int32 CurveIndex, float CurveValue, FFrameValues& CurveValuesPerFrame, const TArray<FFrameTime>& FrameTimesInOrder, int32 FrameIndex, const FFramePoses& Poses, const TArray<FFrameTime>& FrameTimes, const int32 ActualFrameIndex);
 };
diff --git a/Engine/Plugins/Experimental/Animation/RigMapper/Source/RigMapperEditor/RigMapperEditor.Build.cs b/Engine/Plugins/Experimental/Animation/RigMapper/Source/RigMapperEditor/RigMapperEditor.Build.cs
--- a/Engine/Plugins/Experimental/Animation/RigMapper/Source/RigMapperEditor/RigMapperEditor.Build.cs
+++ b/Engine/Plugins/Experimental/Animation/RigMapper/Source/RigMapperEditor/RigMapperEditor.Build.cs
@@ -22,7 +22,6 @@ public class RigMapperEditor : ModuleRules
 				"CoreUObject",
 				"AnimGraph",
 				"BlueprintGraph",
-				"RigMapper",
 				"ControlRig",
 				"LevelSequence",
 				"MovieScene",
@@ -35,7 +34,8 @@ public class RigMapperEditor : ModuleRules
 				"ContentBrowser",
 				"MessageLog",
 				"InputCore",
-				"Sequencer"
+				"Sequencer",
+				"RigMapper"
 			});
 	}
 }
