[RigMapper]:
- Fix output values generation when not using the processor
- Fix value range warning triggered when comparing floats (i.e. 1.000001f >= 1.0f)
- Add tests for the processor, and clearer log for some tests

diff --git a/Engine/Plugins/Experimental/Animation/RigMapper/Source/RigMapper/Private/RigMapper.cpp b/Engine/Plugins/Experimental/Animation/RigMapper/Source/RigMapper/Private/RigMapper.cpp
--- a/Engine/Plugins/Experimental/Animation/RigMapper/Source/RigMapper/Private/RigMapper.cpp
+++ b/Engine/Plugins/Experimental/Animation/RigMapper/Source/RigMapper/Private/RigMapper.cpp
@@ -91,13 +91,13 @@ namespace FacialRigMapping
 	{
 		TMap<FString, double> OutValues;
 	
-		for (int32 i = 0; i < OutputNodes.Num(); i++)
+		for (int32 NodeIndex = 0; NodeIndex < OutputNodes.Num(); NodeIndex++)
 		{
 			double OutputValue = 0.0f;
 			
-			if (OutputNodes[i]->TryGetValue(OutputValue) || !bSkipUnset)
+			if (OutputNodes[NodeIndex]->TryGetValue(OutputValue) || !bSkipUnset)
 			{
-				OutValues.Add(OutputNames[i], OutputValue);
+				OutValues.Add(OutputNames[NodeIndex], OutputValue);
 			}
 		}
 
@@ -117,27 +117,31 @@ namespace FacialRigMapping
 			OutValues.AddZeroed(OutputNodes.Num());
 		}
 	
-		for (int32 i = 0; i < OutputNodes.Num(); i++)
+		for (int32 NodeIndex = 0; NodeIndex < OutputNodes.Num(); NodeIndex++)
 		{
-			OutValues.Add(OutputNodes[i]->GetValue());
+			OutValues[NodeIndex] = OutputNodes[NodeIndex]->GetValue();
 		}
 	}
 	
 	void FRigMapper::GetOptionalOutputValuesInOrder(TArray<TOptional<double>>& OutValues) const
 	{
-
+		if (OutValues.Num() != OutputNodes.Num())
+		{
+			OutValues.Reset(OutputNodes.Num());
+			OutValues.AddDefaulted(OutputNodes.Num());
+		}
 	
-		for (int32 i = 0; i < OutputNodes.Num(); i++)
+		for (int32 NodeIndex = 0; NodeIndex < OutputNodes.Num(); NodeIndex++)
 		{
 			double OutputValue;
 			
-			if (OutputNodes[i]->TryGetValue(OutputValue))
+			if (OutputNodes[NodeIndex]->TryGetValue(OutputValue))
 			{
-				OutValues[i] = OutputValue;
+				OutValues[NodeIndex] = OutputValue;
 			}
 			else
 			{
-				OutValues[i].Reset();
+				OutValues[NodeIndex].Reset();
 			}
 		}
 	}
@@ -150,17 +154,17 @@ namespace FacialRigMapping
 			OutValues.AddDefaulted(OutputNodes.Num());
 		}
 	
-		for (int32 i = 0; i < OutputNodes.Num(); i++)
+		for (int32 NodeIndex = 0; NodeIndex < OutputNodes.Num(); NodeIndex++)
 		{
 			double OutputValue;
 			
-			if (OutputNodes[i]->TryGetValue(OutputValue))
+			if (OutputNodes[NodeIndex]->TryGetValue(OutputValue))
 			{
-				OutValues[i] = OutputValue;
+				OutValues[NodeIndex] = OutputValue;
 			}
 			else
 			{
-				OutValues[i].Reset();
+				OutValues[NodeIndex].Reset();
 			}
 		}
 	}
diff --git a/Engine/Plugins/Experimental/Animation/RigMapper/Source/RigMapper/Private/RigMapperDefinition.cpp b/Engine/Plugins/Experimental/Animation/RigMapper/Source/RigMapper/Private/RigMapperDefinition.cpp
--- a/Engine/Plugins/Experimental/Animation/RigMapper/Source/RigMapper/Private/RigMapperDefinition.cpp
+++ b/Engine/Plugins/Experimental/Animation/RigMapper/Source/RigMapper/Private/RigMapperDefinition.cpp
@@ -157,9 +157,13 @@ bool FRigMapperWsFeature::IsValid(const TArray<FString>& InputNames, bool bWarn)
 	{
 		UE_LOG(LogRigMapper, Warning, TEXT("Total Weights for feature %s add up to 0"), *Name)
 	}
-	if (bWarn && (TotalWeight < -1 || TotalWeight > 1) && (!Range.bHasLowerBound || !Range.bHasUpperBound))
+	if (bWarn && !Range.bHasLowerBound && TotalWeight < -1.000001f)
 	{
-		UE_LOG(LogRigMapper, Warning, TEXT("Total Weights for feature %s are quite high (%f) even though a range was not set"), *Name, TotalWeight)
+		UE_LOG(LogRigMapper, Warning, TEXT("Total Weights for feature %s are quite low (%f) even though a lower range bound was not set"), *Name, TotalWeight)
+	}
+	if (bWarn && !Range.bHasUpperBound && TotalWeight > 1.000001f)
+	{
+		UE_LOG(LogRigMapper, Warning, TEXT("Total Weights for feature %s are quite high (%f) even though an upper range bound was not set"), *Name, TotalWeight)
 	}
 	if (Range.bHasLowerBound && Range.bHasUpperBound && Range.LowerBound > Range.UpperBound)
 	{
diff --git a/Engine/Plugins/Experimental/Animation/RigMapper/Source/RigMapperDeveloper/Private/Tests/RigMapperTests.cpp b/Engine/Plugins/Experimental/Animation/RigMapper/Source/RigMapperDeveloper/Private/Tests/RigMapperTests.cpp
--- a/Engine/Plugins/Experimental/Animation/RigMapper/Source/RigMapperDeveloper/Private/Tests/RigMapperTests.cpp
+++ b/Engine/Plugins/Experimental/Animation/RigMapper/Source/RigMapperDeveloper/Private/Tests/RigMapperTests.cpp
@@ -3,6 +3,7 @@
 #include "Misc/AutomationTest.h"
 #include "RigMapperDefinition.h"
 #include "RigMapper.h"
+#include "RigMapperProcessor.h"
 #include "Misc/Paths.h"
 #include "UObject/UObjectGlobals.h"
 #include "Animation/AnimSequence.h"
@@ -150,7 +151,9 @@ TArray<URigMapperDefinition*> CreateInvalidRigMapperDefinitions2()
 IMPLEMENT_SIMPLE_AUTOMATION_TEST(FRigMapperDefinitionTest, "RigMapper.RigMapperDefinition", EAutomationTestFlags::EditorContext | EAutomationTestFlags::EngineFilter)
 
 bool FRigMapperDefinitionTest::RunTest(const FString& InParameters)
-{	
+{
+	UE_LOG(LogRigMapperTest, Log, TEXT("Starting RigMapperDefinitionsTest1: No errors or warnings expected"));
+	
 	TArray<URigMapperDefinition*> Definitions = CreateValidRigMapperDefinitions1();
 	URigMapperDefinition* RigMapperDefinition1 = Definitions[0];
 
@@ -160,6 +163,8 @@ bool FRigMapperDefinitionTest::RunTest(const FString& InParameters)
 		return false;
 	}
 
+	UE_LOG(LogRigMapperTest, Log, TEXT("Checking validity"));
+
 	// check that asset is valid
 	if (!RigMapperDefinition1->IsDefinitionValid())
 	{
@@ -167,6 +172,8 @@ bool FRigMapperDefinitionTest::RunTest(const FString& InParameters)
 		return false;
 	}
 
+	UE_LOG(LogRigMapperTest, Log, TEXT("Checking validation"));
+	
 	// also check Validate method
 	if (!RigMapperDefinition1->Validate())
 	{
@@ -174,6 +181,8 @@ bool FRigMapperDefinitionTest::RunTest(const FString& InParameters)
 		return false;
 	}
 
+	UE_LOG(LogRigMapperTest, Log, TEXT("Checking Json Export"));
+
 	// test that we can save and re-load the asset from a json string
 	FString JsonString1;
 	bool bExported = RigMapperDefinition1->ExportAsJsonString(JsonString1);
@@ -183,6 +192,8 @@ bool FRigMapperDefinitionTest::RunTest(const FString& InParameters)
 		return false;
 	}
 
+	UE_LOG(LogRigMapperTest, Log, TEXT("Checking Json import from previous export"));
+	
 	// load back in to a second definition
 	URigMapperDefinition* RigMapperDefinition2 = NewObject<URigMapperDefinition>((UObject*)GetTransientPackage(), NAME_None, RF_Transient);
 	bool bImported = RigMapperDefinition2->LoadFromJsonString(JsonString1);
@@ -191,6 +202,9 @@ bool FRigMapperDefinitionTest::RunTest(const FString& InParameters)
 		UE_LOG(LogRigMapperTest, Error, TEXT("Failed to import RigMapperDefinition asset from a Json String"));
 		return false;
 	}
+
+	UE_LOG(LogRigMapperTest, Log, TEXT("Checking Json export from previous import"));
+	
 	FString JsonString2;
 	bExported = RigMapperDefinition2->ExportAsJsonString(JsonString2);
 	if (!bExported)
@@ -241,6 +255,8 @@ bool FRigMapperDefinitionTest::RunTest(const FString& InParameters)
 		return false;
 	}
 
+	UE_LOG(LogRigMapperTest, Log, TEXT("Checking Json file export"));
+	
 	// repeat the test above but from file
 	FString Path = FPaths::ProjectSavedDir() / TEXT("test_export.json");
 	bExported = RigMapperDefinition1->ExportAsJsonFile({ Path });
@@ -250,6 +266,8 @@ bool FRigMapperDefinitionTest::RunTest(const FString& InParameters)
 		return false;
 	}
 
+	UE_LOG(LogRigMapperTest, Log, TEXT("Checking Json file export round trip"));
+
 	URigMapperDefinition* RigMapperDefinition3 = NewObject<URigMapperDefinition>((UObject*)GetTransientPackage(), NAME_None, RF_Transient);
 	bImported = RigMapperDefinition3->LoadFromJsonString(JsonString1);
 	if (!bImported)
@@ -275,9 +293,12 @@ bool FRigMapperDefinitionTest::RunTest(const FString& InParameters)
 		return false;
 	}
 
+
+	UE_LOG(LogRigMapperTest, Log, TEXT("Starting RigMapperDefinitionsTest2: Error(s) and Warning(s) expected"));
+	
 	// construct a couple of invalid definitions and check these are flagged as invalid
 	TArray<URigMapperDefinition*> Definitions2 = CreateInvalidRigMapperDefinitions2();
-
+	
 	if (Definitions2[0]->IsDefinitionValid())
 	{
 		UE_LOG(LogRigMapperTest, Error, TEXT("RigMapperDefinition should be invalid"));
@@ -483,13 +504,169 @@ bool FRigMapperTest::RunTest(const FString& InParameters)
 	//		OutputVal11 = TestSDKFeature3 = 0.35
 
 	const double* OutputVal9 = Outputs.Find(FString(TEXT("OutputVal9")));
+	if (!OutputVal9)
+	{
+		UE_LOG(LogRigMapperTest, Error, TEXT("Could not find Output Value 9 (TestWSFeature3)"));
+		return false;
+	}
+	if (!FMath::IsNearlyEqual(*OutputVal9, 0.5755, SMALL_NUMBER))
+	{
+		UE_LOG(LogRigMapperTest, Error, TEXT("Invalid value for Output Value 9 (TestWSFeature3): %f instead of 0.5755"), *OutputVal9);
+		return false;
+	}
 	const double* OutputVal10 = Outputs.Find(FString(TEXT("OutputVal10")));
+	if (!OutputVal10)
+	{
+		UE_LOG(LogRigMapperTest, Error, TEXT("Could not find Output Value 10 (TestMultFeature3)"));
+		return false;
+	}
+	if (!FMath::IsNearlyEqual(*OutputVal10, 0.096, SMALL_NUMBER))
+	{
+		UE_LOG(LogRigMapperTest, Error, TEXT("Invalid value for Output Value 10 (TestMultFeature3): %f instead of 0.096"), *OutputVal10);
+		return false;
+	}
 	const double* OutputVal11 = Outputs.Find(FString(TEXT("OutputVal11")));
+	if (!OutputVal11)
+	{
+		UE_LOG(LogRigMapperTest, Error, TEXT("Could not find Output Value 11 (TestSdkFeature3)"));
+		return false;
+	}
+	if (!FMath::IsNearlyEqual(*OutputVal11, 0.35, SMALL_NUMBER))
+	{
+		UE_LOG(LogRigMapperTest, Error, TEXT("Invalid value for Output Value 11 (TestSdkFeature3): %f instead of 0.35"), *OutputVal11);
+		return false;
+	}
 
-	if (!OutputVal9 || !OutputVal10 || !OutputVal11 || !FMath::IsNearlyEqual(*OutputVal9, 0.5755, SMALL_NUMBER)
-		|| !FMath::IsNearlyEqual(*OutputVal10, 0.096, SMALL_NUMBER) || !FMath::IsNearlyEqual(*OutputVal11, 0.35, SMALL_NUMBER))
+	// TODO would also be good to test this in the anim node but beyond the scope of the current tests
+	return true;
+}
+
+IMPLEMENT_SIMPLE_AUTOMATION_TEST(FRigMapperProcessorTest, "RigMapper.RigMapperProcessor", EAutomationTestFlags::EditorContext | EAutomationTestFlags::EngineFilter)
+
+
+bool FRigMapperProcessorTest::RunTest(const FString& InParameters)
+{
+	const TArray<URigMapperDefinition*> RigMapperDefinitions = CreateValidRigMapperDefinitions1();
+	if (RigMapperDefinitions.Num() != 2)
+	{
+		UE_LOG(LogRigMapperTest, Error, TEXT("Failed to create RigMapperDefinition assets"));
+		return false;
+	}
+
+	FRigMapperProcessor Processor(RigMapperDefinitions);
+	if (!Processor.IsValid())
 	{
-		UE_LOG(LogRigMapperTest, Error, TEXT("Output values for tested rig-mapping are not as expected"));
+		UE_LOG(LogRigMapperTest, Error, TEXT("Failed to initialize RigMapperProcessor"));
+		return false;
+	}
+
+	// set up some sample inputs
+	TArray<FString> InputNames;
+	InputNames.Add(TEXT("InputVal1"));
+	InputNames.Add(TEXT("InputVal2"));
+	InputNames.Add(TEXT("InputVal3"));
+	InputNames.Add(TEXT("InputVal4"));
+	InputNames.Add(TEXT("InputVal5"));
+
+	FRigMapperProcessor::FPoseValues InputValues;
+	InputValues.Add(0.25);
+	InputValues.Add(0.4);
+	InputValues.Add(0.5);
+	InputValues.Add(0.6);
+	InputValues.Add(0.75);
+
+	// evaluate the rig mappers; this evaluates the whole thing in the same way in which this is done in the anim node
+	FRigMapperProcessor::FPoseValues OutputValues;
+	if (!Processor.EvaluateFrame(InputNames, InputValues, OutputValues))
+	{
+		UE_LOG(LogRigMapperTest, Error, TEXT("RigMapperProcessor failed to evaluate input values"));
+		return false;
+	}
+	const TArray<FString>& OutputNames = Processor.GetOutputNames();
+
+	// check that the outputs are correct
+
+	// Inputs
+	//	InputVal1 = 0.25
+	//	InputVal2 = 0.4
+	//	InputVal3 = 0.5
+	//	InputVal4 = 0.6
+	//	InputVal5 = 0.75
+	// 
+	// Layer1
+	//	TestWSFeature1 = InputVal1 * 0.25 + InputVal2 * 0.25 + InputVal3 * 0.5 = 0.25 * 0.25 + 0.4 * 0.25 + 0.5 * 0.5 = 0.4125
+	//	TestWSFeature2 = InputVal4 * 0.4 + InputVal5 * 0.6 = 0.6 * 0.4 + 0.75 * 0.6 = 0.69
+	//	TestMultFeature1 = InputVal1 * InputVal2 = 0.25 * 0.4 = 0.1
+	//	TestMultFeature2 = InputVal2 * InputVal3 = 0.4 * 0.5 = 0.2
+	//	TestSDKFeature1 = InputVal2 SDK values (0,0) (0.5,0.6) (1.0,1.0) = 0.48
+	//	TestSDKFeature2 = InputVal5 SDK values (0,0.25) (1,0.5) = 0.4375
+	// 
+	//	Outputs:
+	// 		OutputVal1 = TestWSFeature1 = 0.4125
+	//		OutputVal2 = TestWSFeature2 = 0.69
+	//		OutputVal3 = TestMultFeature1 = 0.1
+	//		OutputVal4 = TestMultFeature2 = 0.2
+	//		OutputVal5 = TestSDKFeature1 = 0.48
+	//		OutputVal6 = TestSDKFeature2 = 0.4375
+	// 
+	// Layer2
+	//	TestWSFeature3 = OutputVal1 * 0.2 + OutputVal2 * 0.7 + OutputVal3 * 0.1 = 0.4125 * 0.2 + 0.69 * 0.7 + 0.1 * 0.1 = 0.0925
+	//	TestMultFeature3 = OutputVal4 * OutputVal5 = 0.096
+	//	TestSDKFeature3 = OutputVal6 SDK values (0,0) (1.0, 0.8) = 0.35
+	//
+	//	Outputs:
+	//		OutputVal9 = TestWSFeature3 = 0.5755
+	//		OutputVal10 = TestMultFeature3 = 0.096
+	//		OutputVal11 = TestSDKFeature3 = 0.35
+
+	int32 OutputIndex = OutputNames.Find(FString(TEXT("OutputVal9")));
+	if (!OutputValues.IsValidIndex(OutputIndex))
+	{
+		UE_LOG(LogRigMapperTest, Error, TEXT("Invalid index for Output Value 9 (TestWSFeature3): %d"), OutputIndex);
+		return false;
+	}
+	if (!OutputValues[OutputIndex].IsSet())
+	{
+		UE_LOG(LogRigMapperTest, Error, TEXT("Invalid value for Output Value 9 (TestWSFeature3)"));
+		return false;
+	}
+	if (!FMath::IsNearlyEqual(OutputValues[OutputIndex].GetValue(), 0.5755, 0.00001))
+	{
+		UE_LOG(LogRigMapperTest, Error, TEXT("Invalid value for Output Value 9 (TestWSFeature3): %f instead of 0.5755"), OutputValues[OutputIndex].GetValue());
+		return false;
+	}
+
+	OutputIndex = OutputNames.Find(FString(TEXT("OutputVal10")));
+	if (!OutputValues.IsValidIndex(OutputIndex))
+	{
+		UE_LOG(LogRigMapperTest, Error, TEXT("Invalid index for Output Value 10 (TestMultFeature3): %d"), OutputIndex);
+		return false;
+	}
+	if (!OutputValues[OutputIndex].IsSet())
+	{
+		UE_LOG(LogRigMapperTest, Error, TEXT("Invalid value for Output Value 10 (TestMultFeature3)"));
+		return false;
+	}
+	if (!FMath::IsNearlyEqual(OutputValues[OutputIndex].GetValue(), 0.096, 0.00001))
+	{
+		UE_LOG(LogRigMapperTest, Error, TEXT("Invalid value for Output Value 10 (TestMultFeature3): %f instead of 0.096"), OutputValues[OutputIndex].GetValue());
+		return false;
+	}
+
+	OutputIndex = OutputNames.Find(FString(TEXT("OutputVal11")));
+	if (!OutputValues.IsValidIndex(OutputIndex))
+	{
+		UE_LOG(LogRigMapperTest, Error, TEXT("Invalid index for Output Value 11 (TestSdkFeature3): %d"), OutputIndex);
+		return false;
+	}
+	if (!OutputValues[OutputIndex].IsSet())
+	{
+		UE_LOG(LogRigMapperTest, Error, TEXT("Invalid value for Output Value 11 (TestSdkFeature3)"));
+		return false;
+	}
+	if (!FMath::IsNearlyEqual(OutputValues[OutputIndex].GetValue(), 0.35, 0.00001))
+	{
+		UE_LOG(LogRigMapperTest, Error, TEXT("Invalid value for Output Value 11 (TestSdkFeature3): %f instead of 0.35"), OutputValues[OutputIndex].GetValue());
 		return false;
 	}
 
@@ -501,6 +678,8 @@ IMPLEMENT_SIMPLE_AUTOMATION_TEST(RigMapperLinkedDefinitionsTest, "RigMapper.RigM
 
 bool RigMapperLinkedDefinitionsTest::RunTest(const FString& InParameters)
 {
+	UE_LOG(LogRigMapperTest, Log, TEXT("Starting RigMapperLinkedDefinitionsTest1: No error or warning expected"));
+	
 	// create two rigmapper definitions
 	TArray<URigMapperDefinition*> RigMapperDefinitions1 = CreateValidRigMapperDefinitions1();
 	if (RigMapperDefinitions1.Num() != 2)
@@ -595,6 +774,8 @@ bool RigMapperLinkedDefinitionsTest::RunTest(const FString& InParameters)
 
 	// no explicit test for GetBakedInputRec here
 
+	UE_LOG(LogRigMapperTest, Log, TEXT("Starting RigMapperLinkedDefinitionsTest2: No error or warning expected"));
+	
 	// test for a case where one of the NullOutputs is used; this should still be valid
 	TArray<URigMapperDefinition*> RigMapperDefinitions2 = CreateValidRigMapperDefinitions2();
 
@@ -629,6 +810,8 @@ bool RigMapperLinkedDefinitionsTest::RunTest(const FString& InParameters)
 
 	// now test for a couple of invalid cases and check that validation catches these
 
+	UE_LOG(LogRigMapperTest, Log, TEXT("Starting RigMapperLinkedDefinitionsTest3: Error(s) and Warning(s) expected"));
+
 	// case 1 is definitions which are valid individually but are missing an input to the second set of definitions
 	TArray<URigMapperDefinition*> RigMapperDefinitions3 = CreateInvalidRigMapperDefinitions1();
 
@@ -661,6 +844,8 @@ bool RigMapperLinkedDefinitionsTest::RunTest(const FString& InParameters)
 		return false;
 	}
 
+	UE_LOG(LogRigMapperTest, Log, TEXT("Starting RigMapperLinkedDefinitionsTest4: Error(s) and Warning(s) expected"));
+	
 	// case 2 is simply an example where the individual definitions are invalid
 	TArray<URigMapperDefinition*> RigMapperDefinitions4 = CreateInvalidRigMapperDefinitions1();
 
