RigMapper speedup of initialization and re-enabling RigMapper

#changelist validated
#virtualized

diff --git a/Engine/Plugins/Experimental/Animation/RigMapper/Source/RigMapper/Private/AnimNode_RigMapper.cpp b/Engine/Plugins/Experimental/Animation/RigMapper/Source/RigMapper/Private/AnimNode_RigMapper.cpp
--- a/Engine/Plugins/Experimental/Animation/RigMapper/Source/RigMapper/Private/AnimNode_RigMapper.cpp
+++ b/Engine/Plugins/Experimental/Animation/RigMapper/Source/RigMapper/Private/AnimNode_RigMapper.cpp
@@ -228,25 +228,25 @@ bool FAnimNode_RigMapper::InitializeRigMapping(USkeletalMesh* TargetMesh)
 	LoadedDefinitions = MoveTemp(DefinitionsToLoad);
 	
 	// Cache a map of curve indices to bulk get current curve values for the current pose
-	const TArray<FString>& InputNames = RigMapperProcessor.GetInputNames();
+	const TArray<FName>& InputNames = RigMapperProcessor.GetInputNames();
 	InputCurveMappings.Empty();
 	InputCurveMappings.Reserve(InputNames.Num());
 	for (int32 InputIndex = 0; InputIndex < InputNames.Num(); InputIndex++)
 	{
-		InputCurveMappings.Add(*InputNames[InputIndex], InputIndex);
+		InputCurveMappings.Add(InputNames[InputIndex], InputIndex);
 	}
 
 	// Cache a map of curve indices to bulk set the new curve values for the output pose
-	const TArray<FString>& OutputNames = RigMapperProcessor.GetOutputNames();
+	const TArray<FName>& OutputNames = RigMapperProcessor.GetOutputNames();
 	OutputCurveMappings.Empty();
 	OutputCurveMappings.Reserve(OutputNames.Num());
 	for (int32 OutputIndex = 0; OutputIndex < OutputNames.Num(); OutputIndex++)
 	{
-		const FString& CurveName = OutputNames[OutputIndex];
+		const FName& CurveName = OutputNames[OutputIndex];
 
 		// We need to cache the matching input curve index to allow lerping depending on the node's Alpha
 		int32 InputIndex = InputNames.Find(CurveName);
-		OutputCurveMappings.Add(*OutputNames[OutputIndex], OutputIndex, InputIndex);
+		OutputCurveMappings.Add(OutputNames[OutputIndex], OutputIndex, InputIndex);
 	}
 	
 	return true;
diff --git a/Engine/Plugins/Experimental/Animation/RigMapper/Source/RigMapper/Private/RigMapper.cpp b/Engine/Plugins/Experimental/Animation/RigMapper/Source/RigMapper/Private/RigMapper.cpp
--- a/Engine/Plugins/Experimental/Animation/RigMapper/Source/RigMapper/Private/RigMapper.cpp
+++ b/Engine/Plugins/Experimental/Animation/RigMapper/Source/RigMapper/Private/RigMapper.cpp
@@ -22,40 +22,54 @@ namespace FacialRigMapping
     	InputNodes.Empty();
     	InputNames.Empty();
 	}
-	
+
 	bool FRigMapper::LoadDefinition(const URigMapperDefinition* Definition)
 	{
 		Reset();
-		
-		for (const FString& InputName : Definition->Inputs)
+
+		const int32 NumNodes = Definition->Inputs.Num()
+			+ Definition->Features.Multiply.Num()
+			+ Definition->Features.WeightedSums.Num()
+			+ Definition->Features.SDKs.Num();
+
+		Nodes.Reserve(NumNodes);
+
+		InputNodes.Reserve(Definition->Inputs.Num());
+		InputNames.Reserve(Definition->Inputs.Num());
+		for (int32 Index = 0; Index < Definition->Inputs.Num(); ++Index)
 		{
-			TSharedPtr<FEvalInputNode> Node = MakeShared<FEvalInputNode>(InputName);
-			
+			FName InputName = FName(*Definition->Inputs[Index]);
+			FNodePtr Node{ ENodeType::Input, Index};
+
 			InputNodes.Add(Node);
 			InputNames.Add(InputName);
 
 			Nodes.Add(InputName, Node);
 		}
 
-		for (const FRigMapperMultiplyFeature& Feature : Definition->Features.Multiply)
+		for (int32 Index = 0; Index < Definition->Features.Multiply.Num(); ++Index)
 		{
-			Nodes.Add(Feature.Name, MakeShared<FEvalNodeMultiply>(Feature));
+			Nodes.Add(FName(*Definition->Features.Multiply[Index].Name), FNodePtr{ENodeType::Multiply, Index});
 		}
-		for (const FRigMapperWsFeature& Feature : Definition->Features.WeightedSums)
+
+		for (int32 Index = 0; Index < Definition->Features.WeightedSums.Num(); ++Index)
 		{
-			Nodes.Add(Feature.Name, MakeShared<FEvalNodeWeightedSum>(Feature));
+			Nodes.Add(FName(*Definition->Features.WeightedSums[Index].Name), FNodePtr{ ENodeType::WeightedSum, Index});
 		}
-		for (const FRigMapperSdkFeature& Feature : Definition->Features.SDKs)
+
+		for (int32 Index = 0; Index < Definition->Features.SDKs.Num(); ++Index)
 		{
-			Nodes.Add(Feature.Name, MakeShared<FEvalNodePiecewiseLinear>(Feature));
+			Nodes.Add(FName(*Definition->Features.SDKs[Index].Name), FNodePtr{ ENodeType::PiecewiseLinear, Index});
 		}
 
+		OutputNodes.Reserve(Definition->Outputs.Num());
+		OutputNames.Reserve(Definition->Outputs.Num());
 		for (const TPair<FString, FString>& Output : Definition->Outputs)
 		{
-			const FString OutputName = Output.Key;
-			const FString LinkedNodeName = Output.Value;
-			
-			if (const EvalPtr* Node = Nodes.Find(LinkedNodeName))
+			const FName OutputName = FName(*Output.Key);
+			const FName LinkedNodeName = FName(*Output.Value);
+
+			if (const FNodePtr* Node = Nodes.Find(LinkedNodeName))
 			{
 				OutputNodes.Add(*Node);
 				OutputNames.Add(OutputName);
@@ -64,11 +78,30 @@ namespace FacialRigMapping
 
 		if (IsValid())
 		{
-			for (const TPair<FString, EvalPtr>& NamedNode : Nodes)
+			NodeCollection.InputNodes.AddDefaulted(Definition->Inputs.Num());
+
+			NodeCollection.MultiplyNodes.AddDefaulted(Definition->Features.Multiply.Num());
+			for (int32 Index = 0; Index < NodeCollection.MultiplyNodes.Num(); ++Index)
+			{
+				const FRigMapperMultiplyFeature& Feature = Definition->Features.Multiply[Index];
+				NodeCollection.MultiplyNodes[Index].Initialize(Feature, Nodes);
+			}
+
+			NodeCollection.WeightedSumNodes.AddDefaulted(Definition->Features.WeightedSums.Num());
+			for (int32 Index = 0; Index < NodeCollection.WeightedSumNodes.Num(); ++Index)
+			{
+				const FRigMapperWsFeature& Feature = Definition->Features.WeightedSums[Index];
+				NodeCollection.WeightedSumNodes[Index].Initialize(Feature, Nodes);
+			}
+
+			NodeCollection.PiecewiseLinearNodes.AddDefaulted(Definition->Features.SDKs.Num());
+			for (int32 Index = 0; Index < NodeCollection.PiecewiseLinearNodes.Num(); ++Index)
 			{
-				NamedNode.Value->Initialize(Nodes);
+				const FRigMapperSdkFeature& Feature = Definition->Features.SDKs[Index];
+				NodeCollection.PiecewiseLinearNodes[Index].Initialize(Feature, Nodes);
 			}
 		}
+
 		return IsValid();
 	}
 
@@ -76,26 +109,26 @@ namespace FacialRigMapping
 	{
 		if (InputNodes.IsValidIndex(InputIndex))
 		{
-			InputNodes[InputIndex]->SetDirect(Value);
+			InputNodes[InputIndex].SetDirect(NodeCollection, Value);
 			return true;
 		}
 		return false;
 	}
 	
-	bool FRigMapper::SetDirectValue(const FString& InputName, double Value)
+	bool FRigMapper::SetDirectValue(const FName& InputName, double Value)
 	{
 		return SetDirectValue(InputNames.Find(InputName), Value);
 	}
 
-	TMap<FString, double> FRigMapper::GetOutputValues(bool bSkipUnset) const
+	TMap<FName, double> FRigMapper::GetOutputValues(bool bSkipUnset) const
 	{
-		TMap<FString, double> OutValues;
+		TMap<FName, double> OutValues;
+		OutValues.Reserve(OutputNodes.Num());
 	
 		for (int32 NodeIndex = 0; NodeIndex < OutputNodes.Num(); NodeIndex++)
 		{
 			double OutputValue = 0.0f;
-			
-			if (OutputNodes[NodeIndex]->TryGetValue(OutputValue) || !bSkipUnset)
+			if (OutputNodes[NodeIndex].TryGetValue(NodeCollection, OutputValue) || !bSkipUnset)
 			{
 				OutValues.Add(OutputNames[NodeIndex], OutputValue);
 			}
@@ -104,7 +137,7 @@ namespace FacialRigMapping
 		return OutValues;
 	}
 
-	const TArray<FString>& FRigMapper::GetOutputNames() const
+	const TArray<FName>& FRigMapper::GetOutputNames() const
 	{
 		return OutputNames;
 	}
@@ -116,13 +149,13 @@ namespace FacialRigMapping
 			OutValues.Reset(OutputNodes.Num());
 			OutValues.AddZeroed(OutputNodes.Num());
 		}
-	
+
 		for (int32 NodeIndex = 0; NodeIndex < OutputNodes.Num(); NodeIndex++)
 		{
-			OutValues[NodeIndex] = OutputNodes[NodeIndex]->GetValue();
+			OutValues[NodeIndex] = OutputNodes[NodeIndex].GetValue(NodeCollection);
 		}
 	}
-	
+
 	void FRigMapper::GetOptionalOutputValuesInOrder(TArray<TOptional<double>>& OutValues) const
 	{
 		if (OutValues.Num() != OutputNodes.Num())
@@ -130,12 +163,12 @@ namespace FacialRigMapping
 			OutValues.Reset(OutputNodes.Num());
 			OutValues.AddDefaulted(OutputNodes.Num());
 		}
-	
+
 		for (int32 NodeIndex = 0; NodeIndex < OutputNodes.Num(); NodeIndex++)
 		{
 			double OutputValue;
-			
-			if (OutputNodes[NodeIndex]->TryGetValue(OutputValue))
+
+			if (OutputNodes[NodeIndex].TryGetValue(NodeCollection, OutputValue))
 			{
 				OutValues[NodeIndex] = OutputValue;
 			}
@@ -153,12 +186,12 @@ namespace FacialRigMapping
 			OutValues.Reset(OutputNodes.Num());
 			OutValues.AddDefaulted(OutputNodes.Num());
 		}
-	
+
 		for (int32 NodeIndex = 0; NodeIndex < OutputNodes.Num(); NodeIndex++)
 		{
 			double OutputValue;
-			
-			if (OutputNodes[NodeIndex]->TryGetValue(OutputValue))
+
+			if (OutputNodes[NodeIndex].TryGetValue(NodeCollection, OutputValue))
 			{
 				OutValues[NodeIndex] = OutputValue;
 			}
@@ -169,71 +202,103 @@ namespace FacialRigMapping
 		}
 	}
 
-	void FRigMapper::SetDirty() const
+	void FRigMapper::SetDirty()
 	{
-		for (const auto& N : Nodes)
+		for (TPair<FName, FNodePtr>& Node : Nodes)
 		{
-			N.Value->Reset();
+			Node.Value.Reset(NodeCollection);
 		}
 	}
 
-	const TArray<FString>& FRigMapper::GetInputNames() const
+	const TArray<FName>& FRigMapper::GetInputNames() const
 	{
 		return InputNames;
 	}
-	
-	FEvalNode::FEvalNode(const FString& InName):
-		Name(InName)
-	{
-	}
 
-	double FEvalNode::GetValue() const
-	{
-		double Value = 0.0f;
+	////////////////////////////////////////////////////////////////////////////////////
+	// FNodePtr
+	////////////////////////////////////////////////////////////////////////////////////
 
-		TryGetValue(Value);
-		return Value;
+	FNodePtr::FNodePtr(ENodeType InNodeType, int32 InDataIndex)
+		: DataIndex(InDataIndex)
+		, NodeType(InNodeType)
+	{
 	}
 
-	bool FEvalNode::TryGetValue(double& OutValue) const
+	bool FNodePtr::TryGetValue(const FNodeCollection& NodeCollection, double& OutValue) const
 	{
-		if (!IsInitialized())
-		{
-			return false;
-		}
-		if (!CachedValue.IsSet())
+		switch (NodeType)
 		{
-			if (Evaluate(OutValue))
-			{
-				CachedValue = OutValue;
-				return true;
-			}
-			return false;
+		case ENodeType::Input:
+			return NodeCollection.InputNodes[DataIndex].TryGetValue(NodeCollection, OutValue);
+		case ENodeType::WeightedSum:
+			return NodeCollection.WeightedSumNodes[DataIndex].TryGetValue(NodeCollection, OutValue);
+		case ENodeType::PiecewiseLinear:
+			return NodeCollection.PiecewiseLinearNodes[DataIndex].TryGetValue(NodeCollection, OutValue);
+		case ENodeType::Multiply:
+			return NodeCollection.MultiplyNodes[DataIndex].TryGetValue(NodeCollection, OutValue);
 		}
-		else
+		return false;
+	}
+
+	bool FNodePtr::IsInitialized(const FNodeCollection& NodeCollection) const
+	{
+		switch (NodeType)
 		{
-			OutValue = CachedValue.GetValue();
+		case ENodeType::Input:
 			return true;
+		case ENodeType::WeightedSum:
+			return NodeCollection.WeightedSumNodes[DataIndex].IsInitialized();
+		case ENodeType::PiecewiseLinear:
+			return NodeCollection.PiecewiseLinearNodes[DataIndex].IsInitialized();
+		case ENodeType::Multiply:
+			return NodeCollection.MultiplyNodes[DataIndex].IsInitialized();
 		}
+		return false;
 	}
 
-	void FEvalNode::SetDirect(double Value)
+	double FNodePtr::GetValue(const FNodeCollection& NodeCollection) const
 	{
-		CachedValue = Value;
+		double Value = 0.0;
+		TryGetValue(NodeCollection, Value);
+		return Value;
 	}
 
-	void FEvalNode::Reset()
+	void FNodePtr::SetDirect(FNodeCollection& NodeCollection, double Value)
 	{
-		CachedValue.Reset();
+		switch (NodeType)
+		{
+		case ENodeType::Input:
+			return NodeCollection.InputNodes[DataIndex].SetDirect(Value);
+		case ENodeType::WeightedSum:
+			return NodeCollection.WeightedSumNodes[DataIndex].SetDirect(Value);
+		case ENodeType::PiecewiseLinear:
+			return NodeCollection.PiecewiseLinearNodes[DataIndex].SetDirect(Value);
+		case ENodeType::Multiply:
+			return NodeCollection.MultiplyNodes[DataIndex].SetDirect(Value);
+		}
 	}
 
-	bool FEvalNode::IsInitialized() const
+	void FNodePtr::Reset(FNodeCollection& NodeCollection)
 	{
-		return bInitialized;
+		switch (NodeType)
+		{
+		case ENodeType::Input:
+			return NodeCollection.InputNodes[DataIndex].Reset();
+		case ENodeType::WeightedSum:
+			return NodeCollection.WeightedSumNodes[DataIndex].Reset();
+		case ENodeType::PiecewiseLinear:
+			return NodeCollection.PiecewiseLinearNodes[DataIndex].Reset();
+		case ENodeType::Multiply:
+			return NodeCollection.MultiplyNodes[DataIndex].Reset();
+		}
 	}
 
-	FEvalNodeWeightedSum::FEvalNodeWeightedSum(const FRigMapperWsFeature& FeatureDefinition):
-		FEvalNode(FeatureDefinition.Name)
+	////////////////////////////////////////////////////////////////////////////////////
+	// FEvalNodeWeightedSum
+	////////////////////////////////////////////////////////////////////////////////////
+
+	void FEvalNodeWeightedSum::Initialize(const FRigMapperWsFeature& FeatureDefinition, const TMap<FName, FNodePtr>& Nodes)
 	{
 		if (FeatureDefinition.Range.bHasLowerBound)
 		{
@@ -243,25 +308,35 @@ namespace FacialRigMapping
 		{
 			Range.SetUpperBound(TRange<double>::BoundsType(FeatureDefinition.Range.UpperBound));
 		}
-		WeightedLinkedInputNames = FeatureDefinition.Inputs;
+
+		for (const TPair<FString, double>& Pair : FeatureDefinition.Inputs)
+		{
+			if (const FNodePtr* Node = Nodes.Find(FName(*Pair.Key)))
+			{
+				WeightedLinkedInputs.Emplace(*Node, Pair.Value);
+			}
+		}
+
+		// Will be considered not initialized if any node could not be found
+		bInitialized = WeightedLinkedInputs.Num() == FeatureDefinition.Inputs.Num();
 	}
 
-	bool FEvalNodeWeightedSum::Evaluate(double &OutValue) const
+	bool FEvalNodeWeightedSum::Evaluate(const FNodeCollection& NodeCollection, double& OutValue) const
 	{
 		bool bInputEvaluated = false;
 		OutValue = 0;
 
-		for (const auto& WeightedInput : WeightedLinkedInputs)
+		for (const TPair<FNodePtr, double>& WeightedInput : WeightedLinkedInputs)
 		{
 			double InputValue = 0.0f;
-			
-			if (WeightedInput.Key->TryGetValue(InputValue))
+
+			if (WeightedInput.Key.TryGetValue(NodeCollection, InputValue))
 			{
 				bInputEvaluated = true;
 			}
-			OutValue += WeightedInput.Key->GetValue() * WeightedInput.Value;
+			OutValue += InputValue * WeightedInput.Value;
 		}
-		
+
 		if (Range.HasLowerBound() && OutValue < Range.GetLowerBoundValue())
 		{
 			OutValue = Range.GetLowerBoundValue();
@@ -270,48 +345,38 @@ namespace FacialRigMapping
 		{
 			OutValue = Range.GetUpperBoundValue();
 		}
-		
+
 		return bInputEvaluated;
 	}
 
-	void FEvalNodeWeightedSum::Initialize(const TMap<FString, EvalPtr>& Nodes)
+	////////////////////////////////////////////////////////////////////////////////////
+	// FEvalNodePiecewiseLinear
+	////////////////////////////////////////////////////////////////////////////////////
+
+	void FEvalNodePiecewiseLinear::Initialize(const FRigMapperSdkFeature& FeatureDefinition, const TMap<FName, FNodePtr>& Nodes)
 	{
-		FEvalNode::Initialize(Nodes);
-		WeightedLinkedInputs.Empty(WeightedLinkedInputNames.Num());
+		LinkedInput = FNodePtr{};
 
-		for (const TPair<FString, double>& Pair : WeightedLinkedInputNames)
+		if (const FNodePtr* Node = Nodes.Find(FName(*FeatureDefinition.Input)))
 		{
-			if (const EvalPtr* Node = Nodes.Find(Pair.Key))
-			{
-				if (Node->IsValid())
-				{
-					WeightedLinkedInputs.Add(*Node, Pair.Value);
-				}
-			}
-			// todo: log error?
+			LinkedInput = *Node;
 		}
 
-		// Will be considered not initialized if any node could not be found
-		bInitialized = WeightedLinkedInputs.Num() == WeightedLinkedInputNames.Num();
-	}
-
-	FEvalNodePiecewiseLinear::FEvalNodePiecewiseLinear(const FRigMapperSdkFeature& FeatureDefinition):
-		FEvalNode(FeatureDefinition.Name)
-	{
-		LinkedInputName = FeatureDefinition.Input;
-
+		Keys.Reserve(FeatureDefinition.Keys.Num());
 		for (const FRigMapperSdkKey& Key : FeatureDefinition.Keys)
 		{
-			Keys.Add({ Key.In, Key.Out });
+			Keys.Emplace(Key.In, Key.Out);
 		}
-		Keys.Sort([](const TPair<double,double>& Key1, const TPair<double,double>& Key2) { return Key1.Key < Key2.Key; });
+		Keys.Sort([](const TPair<double, double>& Key1, const TPair<double, double>& Key2) { return Key1.Key < Key2.Key; });
+
+		bInitialized = LinkedInput.IsValid();
 	}
 
-	bool FEvalNodePiecewiseLinear::Evaluate(double &OutValue) const
+	bool FEvalNodePiecewiseLinear::Evaluate(const FNodeCollection& NodeCollection, double& OutValue) const
 	{
 		double InputValue = 0.0f;
 
-		if (!LinkedInput->TryGetValue(InputValue))
+		if (!LinkedInput.TryGetValue(NodeCollection, InputValue))
 		{
 			return false;
 		}
@@ -319,30 +384,30 @@ namespace FacialRigMapping
 		if (!Evaluate_Static(InputValue, Keys, OutValue))
 		{
 			UE_LOG(LogRigMapper, Warning, TEXT("PiecewiseLinear could not calculate the output value"))
-			return false;	
+				return false;
 		}
 
 		return true;
 	}
 
-	bool FEvalNodePiecewiseLinear::Evaluate_Static(const double InputValue, const TArray<TPair<double,double>>& Keys, double &OutValue)
+	bool FEvalNodePiecewiseLinear::Evaluate_Static(const double InputValue, TConstArrayView<TPair<double, double>> Keys, double& OutValue)
 	{
 		if (InputValue <= Keys[0].Key)
 		{
 			OutValue = Keys[0].Value;
 			return true;
 		}
-	
+
 		if (InputValue >= Keys.Last().Key)
 		{
 			OutValue = Keys.Last().Value;
 			return true;
 		}
-	
-		for (int32 KeyIndex = 1 ; KeyIndex < Keys.Num(); KeyIndex++)
+
+		for (int32 KeyIndex = 1; KeyIndex < Keys.Num(); KeyIndex++)
 		{
 			const TPair<double, double>& CurrentKey = Keys[KeyIndex];
-			
+
 			if (InputValue == CurrentKey.Key)
 			{
 				OutValue = CurrentKey.Value;
@@ -351,7 +416,7 @@ namespace FacialRigMapping
 			if (InputValue < CurrentKey.Key)
 			{
 				const TPair<double, double>& PrevKey = Keys[KeyIndex - 1];
-				
+
 				const double Percent = (InputValue - PrevKey.Key) / (CurrentKey.Key - PrevKey.Key);
 
 				OutValue = PrevKey.Value + Percent * (CurrentKey.Value - PrevKey.Value);
@@ -362,25 +427,30 @@ namespace FacialRigMapping
 		return false;
 	}
 
-	void FEvalNodePiecewiseLinear::Initialize(const TMap<FString, EvalPtr>& Nodes)
+	////////////////////////////////////////////////////////////////////////////////////
+	// FEvalNodeMultiply
+	////////////////////////////////////////////////////////////////////////////////////
+
+	void FEvalNodeMultiply::Initialize(const FRigMapperMultiplyFeature& FeatureDefinition, const TMap<FName, FNodePtr>& Nodes)
 	{
-		FEvalNode::Initialize(Nodes);
-		LinkedInput.Reset();
-		
-		if (const EvalPtr* Node = Nodes.Find(LinkedInputName))
+		LinkedInputs.Empty(FeatureDefinition.Inputs.Num());
+
+		for (const FString& InputName : FeatureDefinition.Inputs)
 		{
-			LinkedInput = *Node;
+			if (const FNodePtr* Node = Nodes.Find(FName(*InputName)))
+			{
+				if (Node->IsValid())
+				{
+					LinkedInputs.Add(*Node);
+				}
+			}
 		}
 
-		bInitialized = LinkedInput.IsValid();
-	}
-
-	FEvalNodeMultiply::FEvalNodeMultiply(const FRigMapperMultiplyFeature& FeatureDefinition):
-		FEvalNode(FeatureDefinition.Name), LinkedInputNames(FeatureDefinition.Inputs)
-	{
+		// Will be considered not initialized if any node could not be found
+		bInitialized = LinkedInputs.Num() == FeatureDefinition.Inputs.Num();
 	}
 
-	bool FEvalNodeMultiply::Evaluate(double &OutValue) const
+	bool FEvalNodeMultiply::Evaluate(const FNodeCollection& NodeCollection, double& OutValue) const
 	{
 		if (LinkedInputs.IsEmpty())
 		{
@@ -390,38 +460,17 @@ namespace FacialRigMapping
 		OutValue = 1.0;
 		bool bInputEvaluated = false;
 
-		for (const EvalPtr& Input : LinkedInputs)
+		for (const FNodePtr& Input : LinkedInputs)
 		{
 			double InputValue = 0.0f;
 
-			if (Input->TryGetValue(InputValue))
+			if (Input.TryGetValue(NodeCollection, InputValue))
 			{
 				bInputEvaluated = true;
 			}
-			OutValue *= Input->GetValue();
+			OutValue *= InputValue;
 		}
 
 		return bInputEvaluated;
 	}
-
-	void FEvalNodeMultiply::Initialize(const TMap<FString, EvalPtr>& Nodes)
-	{
-		FEvalNode::Initialize(Nodes);
-		LinkedInputs.Empty(LinkedInputNames.Num());
-
-		for (const FString& InputName : LinkedInputNames)
-		{
-			if (const EvalPtr* Node = Nodes.Find(InputName))
-			{
-				if (Node->IsValid())
-				{
-					LinkedInputs.Add(*Node);
-				}
-			}
-			// todo: log error?
-		}
-
-		// Will be considered not initialized if any node could not be found
-		bInitialized = LinkedInputs.Num() == LinkedInputNames.Num();
-	}
 }
\ No newline at end of file
diff --git a/Engine/Plugins/Experimental/Animation/RigMapper/Source/RigMapper/Private/RigMapperProcessor.cpp b/Engine/Plugins/Experimental/Animation/RigMapper/Source/RigMapper/Private/RigMapperProcessor.cpp
--- a/Engine/Plugins/Experimental/Animation/RigMapper/Source/RigMapper/Private/RigMapperProcessor.cpp
+++ b/Engine/Plugins/Experimental/Animation/RigMapper/Source/RigMapper/Private/RigMapperProcessor.cpp
@@ -4,69 +4,81 @@
 #include "RigMapperLog.h"
 #include "RigMapperDefinition.h"
 
+DEFINE_STAT(STAT_RigMapperGetRigMapper);
+DEFINE_STAT(STAT_RigMapperLoadDefinition);
 
-FRigMapperProcessor::FRigMapperProcessor(const TArray<URigMapperDefinition*>& InDefinitions)
-{
-	LoadDefinitions(InDefinitions, RigMappers);
-	IndexCache.AddDefaulted(RigMappers.Num());
-}
 
-FRigMapperProcessor::FRigMapperProcessor(URigMapperDefinition* InDefinition)
+class FRigMapperDefinitionsSingleton
 {
-	LoadDefinitions({ InDefinition }, RigMappers);
-	IndexCache.AddDefaulted(RigMappers.Num());
-}
-
-FRigMapperProcessor::FRigMapperProcessor(const TArray<TSharedPtr<FacialRigMapping::FRigMapper>>& InRigMappers): RigMappers(InRigMappers)
-{
-	IndexCache.AddDefaulted(RigMappers.Num());	
-}
-
-FRigMapperProcessor::FRigMapperProcessor(TSharedPtr<FacialRigMapping::FRigMapper> InRigMapper): RigMappers( { InRigMapper } )
+public:
+	static FRigMapperDefinitionsSingleton& Get();
+	bool GetRigMapper(const TObjectPtr<const URigMapperDefinition>& InDefinition, FacialRigMapping::FRigMapper& OutRigMapper);
+private:
+	FRigMapperDefinitionsSingleton() = default;
+	virtual ~FRigMapperDefinitionsSingleton() = default;
+	TMap<FSoftObjectPath, FacialRigMapping::FRigMapper> RigMappers;
+};
+FRigMapperDefinitionsSingleton& FRigMapperDefinitionsSingleton::Get()
 {
-	IndexCache.AddDefaulted();
+	static FRigMapperDefinitionsSingleton SingletonInstance;
+	return SingletonInstance;
 }
-
-bool FRigMapperProcessor::LoadDefinitions(const TArray<URigMapperDefinition*>& Definitions, TArray<TSharedPtr<FacialRigMapping::FRigMapper>>& RigMappers)
+bool FRigMapperDefinitionsSingleton::GetRigMapper(const TObjectPtr<const URigMapperDefinition>& InDefinition, FacialRigMapping::FRigMapper& OutRigMapper)
 {
-	if (Definitions.IsEmpty())
+	SCOPE_CYCLE_COUNTER(STAT_RigMapperGetRigMapper);
+	OutRigMapper.Reset();
+	FSoftObjectPath DefinitionSoftObjectPath(InDefinition);
+	if (const FacialRigMapping::FRigMapper* RigMapper = RigMappers.Find(DefinitionSoftObjectPath); RigMapper != nullptr)
 	{
-		UE_LOG(LogRigMapper, Error, TEXT("Not enough definitions to remap"))
-		return false;
+		OutRigMapper = *RigMapper;
 	}
-	
-	RigMappers.Empty(Definitions.Num());
-	
-	for (int32 DefinitionIndex = 0; DefinitionIndex < Definitions.Num(); DefinitionIndex++)
+	else
 	{
-		if (!Definitions[DefinitionIndex])
+		SCOPE_CYCLE_COUNTER(STAT_RigMapperLoadDefinition);
+		if (!OutRigMapper.LoadDefinition(InDefinition))
 		{
-			UE_LOG(LogRigMapper, Error, TEXT("Invalid definition at index: %d"), DefinitionIndex)
-			RigMappers.Empty();
+			UE_LOG(LogRigMapper, Error, TEXT("Could not load definition (%s)"), *InDefinition->GetPathName())
 			return false;
 		}
-
-		// check if the definition is validated ... if not, warn
-		if (!Definitions[DefinitionIndex]->WasDefinitionValidated())
+		if (!InDefinition->WasDefinitionValidated())
 		{
-			UE_LOG(LogRigMapper, Warning, TEXT("RigMapper definition asset at index %d (%s) was not validated please ensure that all RigMapper node definitions have been set to validated assets."), DefinitionIndex, *Definitions[DefinitionIndex]->GetPathName())
+			UE_LOG(LogRigMapper, Warning, TEXT("RigMapper definition asset (%s) was not validated please ensure that all RigMapper node definitions have been set to validated assets."), *InDefinition->GetPathName())
 		}
-		
-		TSharedPtr<FacialRigMapping::FRigMapper> RigMapper = MakeShared<FacialRigMapping::FRigMapper>();
-		
-		if (!RigMapper->LoadDefinition(Definitions[DefinitionIndex]))
+		RigMappers.Add(DefinitionSoftObjectPath, OutRigMapper);
+	}
+	return true;
+}
+
+FRigMapperProcessor::FRigMapperProcessor(const TArray<URigMapperDefinition*>& InDefinitions)
+{
+	for (const URigMapperDefinition* Definition : InDefinitions)
+	{
+		FacialRigMapping::FRigMapper RigMapper;
+		if (!FRigMapperDefinitionsSingleton::Get().GetRigMapper(Definition, RigMapper))
+
 		{
-			UE_LOG(LogRigMapper, Error, TEXT("Could not load definition at index %d (%s)"), DefinitionIndex, *Definitions[DefinitionIndex]->GetPathName())
-			RigMappers.Empty();
-			return false;
+			UE_LOG(LogRigMapper, Error, TEXT("Could not load definition %s"), *Definition->GetPathName())
+				RigMappers.Empty();
+			return;
 		}
-		
-		RigMappers.Add(RigMapper);
+		RigMappers.Add(MoveTemp(RigMapper));
 	}
-	
-	// todo: check definitions chain 
-	
-	return true;
+
+	IndexCache.AddDefaulted(RigMappers.Num());
+}
+
+FRigMapperProcessor::FRigMapperProcessor(URigMapperDefinition* InDefinition)
+{
+	FacialRigMapping::FRigMapper RigMapper;
+	if (!FRigMapperDefinitionsSingleton::Get().GetRigMapper(InDefinition, RigMapper))
+
+	{
+		UE_LOG(LogRigMapper, Error, TEXT("Could not load definition %s"), *InDefinition->GetPathName())
+			RigMappers.Empty();
+		return;
+	}
+	RigMappers.Add(MoveTemp(RigMapper));
+	IndexCache.AddDefaulted(RigMappers.Num());
 }
 
 bool FRigMapperProcessor::IsValid() const
@@ -75,27 +87,20 @@ bool FRigMapperProcessor::IsValid() const
 	{
 		return false;
 	}
-	for (int32 RigIndex = 0; RigIndex < RigMappers.Num(); RigIndex++)
-	{
-		if (!RigMappers[RigIndex])
-		{
-			return false;
-		}
-	}
 	return true;
 }
 
-const TArray<FString>& FRigMapperProcessor::GetInputNames() const
+const TArray<FName>& FRigMapperProcessor::GetInputNames() const
 {
-	return RigMappers[0]->GetInputNames();
+	return RigMappers[0].GetInputNames();
 }
 
-const TArray<FString>& FRigMapperProcessor::GetOutputNames() const
+const TArray<FName>& FRigMapperProcessor::GetOutputNames() const
 {
-	return RigMappers.Last()->GetOutputNames();
+	return RigMappers.Last().GetOutputNames();
 }
 
-bool FRigMapperProcessor::EvaluateFrames(const TArray<FString>& CurveNames, const TArray<FPoseValues>& InFrameValues, TArray<FPoseValues>& OutFrameValues)
+bool FRigMapperProcessor::EvaluateFrames(const TArray<FName>& CurveNames, const TArray<FPoseValues>& InFrameValues, TArray<FPoseValues>& OutFrameValues)
 {
 	if (!IsValid())
 	{
@@ -116,27 +121,27 @@ bool FRigMapperProcessor::EvaluateFrames(const TArray<FString>& CurveNames, cons
 			}
 			else
 			{
-				bValid &= EvaluateFrame_Internal(RigMappers[RigIndex - 1]->GetOutputNames(), OutFrameValues[FrameIndex], OutFrameValues[FrameIndex], RigIndex);
+				bValid &= EvaluateFrame_Internal(RigMappers[RigIndex - 1].GetOutputNames(), OutFrameValues[FrameIndex], OutFrameValues[FrameIndex], RigIndex);
 			}
 		}
 	}
 	return bValid;
 }
 
-bool FRigMapperProcessor::EvaluateFrames(const TArray<FString>& CurveNames, const TArray<FPoseValues>& InFrameValues, TArray<FString>& OutCurveNames,
+bool FRigMapperProcessor::EvaluateFrames(const TArray<FName>& CurveNames, const TArray<FPoseValues>& InFrameValues, TArray<FName>& OutCurveNames,
 	TArray<FPoseValues>& OutFrameValues)
 {
 	if (!IsValid() || !EvaluateFrames(CurveNames, InFrameValues, OutFrameValues))
 	{
 		return false;
 	}
-	const TArray<FString>& NewOutputNames = RigMappers.Last()->GetOutputNames();
+	const TArray<FName>& NewOutputNames = RigMappers.Last().GetOutputNames();
 	OutCurveNames.Reset(NewOutputNames.Num());
 	OutCurveNames.Append(NewOutputNames);
 	return true;
 }
 
-bool FRigMapperProcessor::EvaluateFrames_Interp(const TArray<FString>& CurveNames, const TArray<FPoseValues>& InFrameValues, TArray<FPoseValues>& OutFrameValues, const TArray<FFrameTime>& FrameTimes)
+bool FRigMapperProcessor::EvaluateFrames_Interp(const TArray<FName>& CurveNames, const TArray<FPoseValues>& InFrameValues, TArray<FPoseValues>& OutFrameValues, const TArray<FFrameTime>& FrameTimes)
 {
 	// todo
 	if (!IsValid() || !EvaluateFrames(CurveNames, InFrameValues, OutFrameValues))
@@ -146,7 +151,7 @@ bool FRigMapperProcessor::EvaluateFrames_Interp(const TArray<FString>& CurveName
 	return true;
 }
 
-bool FRigMapperProcessor::EvaluateFrame(const TArray<FString>& CurveNames, const FPoseValues& InCurveValues, FPoseValues& OutCurveValues)
+bool FRigMapperProcessor::EvaluateFrame(const TArray<FName>& CurveNames, const FPoseValues& InCurveValues, FPoseValues& OutCurveValues)
 {
 	if (!IsValid())
 	{
@@ -162,28 +167,28 @@ bool FRigMapperProcessor::EvaluateFrame(const TArray<FString>& CurveNames, const
 		}
 		else
 		{
-			bValid &= EvaluateFrame_Internal(RigMappers[RigIndex - 1]->GetOutputNames(), OutCurveValues, OutCurveValues, RigIndex);
+			bValid &= EvaluateFrame_Internal(RigMappers[RigIndex - 1].GetOutputNames(), OutCurveValues, OutCurveValues, RigIndex);
 		}
 	}
 
 	return true;
 }
 
-bool FRigMapperProcessor::EvaluateFrame(const TArray<FString>& CurveNames, const FPoseValues& InCurveValues, TArray<FString>& OutCurveNames,
+bool FRigMapperProcessor::EvaluateFrame(const TArray<FName>& CurveNames, const FPoseValues& InCurveValues, TArray<FName>& OutCurveNames,
 	FPoseValues& OutCurveValues)
 {
 	if (!IsValid() || !EvaluateFrame(CurveNames, InCurveValues, OutCurveValues))
 	{
 		return false;
 	}
-	const TArray<FString>& NewOutputNames = RigMappers.Last()->GetOutputNames();
+	const TArray<FName>& NewOutputNames = RigMappers.Last().GetOutputNames();
 	OutCurveNames.Reset(NewOutputNames.Num());
 	OutCurveNames.Append(NewOutputNames);
 	return true;
 }
 
 
-bool FRigMapperProcessor::EvaluateFrame_Internal(const TArray<FString>& CurveNames, const FPoseValues& InCurveValues, FPoseValues& OutCurveValues, const int32 RigIndex)
+bool FRigMapperProcessor::EvaluateFrame_Internal(const TArray<FName>& CurveNames, const FPoseValues& InCurveValues, FPoseValues& OutCurveValues, const int32 RigIndex)
 {
 	if (!RigMappers.IsValidIndex(RigIndex) || !RigMappers[RigIndex].IsValid() || CurveNames.Num() != InCurveValues.Num())
 	{
@@ -194,10 +199,10 @@ bool FRigMapperProcessor::EvaluateFrame_Internal(const TArray<FString>& CurveNam
 		IndexCache[RigIndex].Reset(CurveNames.Num());
 	}
 		
-	RigMappers[RigIndex]->SetDirty();
+	RigMappers[RigIndex].SetDirty();
 
 	// Set inputs
-	const TArray<FString>& InputNames = RigMappers[RigIndex]->GetInputNames();
+	const TArray<FName>& InputNames = RigMappers[RigIndex].GetInputNames();
 		
 	for (int32 CurveIndex = 0; CurveIndex < CurveNames.Num(); CurveIndex++)
 	{
@@ -211,10 +216,10 @@ bool FRigMapperProcessor::EvaluateFrame_Internal(const TArray<FString>& CurveNam
 		const int32 InputIndex = IndexCache[RigIndex][CurveIndex];
 		if (InputNames.IsValidIndex(InputIndex) && InCurveValues[CurveIndex].IsSet())
 		{
-			RigMappers[RigIndex]->SetDirectValue(InputIndex, InCurveValues[CurveIndex].GetValue());
+			RigMappers[RigIndex].SetDirectValue(InputIndex, InCurveValues[CurveIndex].GetValue());
 		}
 	}
 	
-	RigMappers[RigIndex]->GetOptionalFloatOutputValuesInOrder(OutCurveValues);
+	RigMappers[RigIndex].GetOptionalFloatOutputValuesInOrder(OutCurveValues);
 	return true;
 }
diff --git a/Engine/Plugins/Experimental/Animation/RigMapper/Source/RigMapper/Public/AnimNode_RigMapper.h b/Engine/Plugins/Experimental/Animation/RigMapper/Source/RigMapper/Public/AnimNode_RigMapper.h
--- a/Engine/Plugins/Experimental/Animation/RigMapper/Source/RigMapper/Public/AnimNode_RigMapper.h
+++ b/Engine/Plugins/Experimental/Animation/RigMapper/Source/RigMapper/Public/AnimNode_RigMapper.h
@@ -16,8 +16,6 @@ struct FRigCurveElement;
 class UControlRig;
 class USkeletalMesh;
 
-DECLARE_STATS_GROUP(TEXT("RigMapper"), STATGROUP_RigMapper, STATCAT_Advanced);
-
 DECLARE_CYCLE_STAT_EXTERN(TEXT("Evaluate Rig Mapper Node"), STAT_RigMapperEvaluateNode, STATGROUP_RigMapper, RIGMAPPER_API);
 DECLARE_CYCLE_STAT_EXTERN(TEXT("Evaluate Rig Mapping"), STAT_RigMapperEvaluateRigs, STATGROUP_RigMapper, RIGMAPPER_API);
 DECLARE_CYCLE_STAT_EXTERN(TEXT("Initialize Rig Mapping"), STAT_RigMapperInitialize, STATGROUP_RigMapper, RIGMAPPER_API);
diff --git a/Engine/Plugins/Experimental/Animation/RigMapper/Source/RigMapper/Public/RigMapper.h b/Engine/Plugins/Experimental/Animation/RigMapper/Source/RigMapper/Public/RigMapper.h
--- a/Engine/Plugins/Experimental/Animation/RigMapper/Source/RigMapper/Public/RigMapper.h
+++ b/Engine/Plugins/Experimental/Animation/RigMapper/Source/RigMapper/Public/RigMapper.h
@@ -11,136 +11,206 @@ class URigMapperDefinition;
 
 namespace FacialRigMapping
 {
-	class FEvalNode;
-	using EvalPtr = TSharedPtr<FEvalNode>;
+	class FNodePtr;
+
+	class FInputNode;
+	class FEvalNodeWeightedSum;
+	class FEvalNodePiecewiseLinear;
+	class FEvalNodeMultiply;
+
+	struct FNodeCollection
+	{
+		TArray<FInputNode> InputNodes;
+		TArray<FEvalNodeWeightedSum> WeightedSumNodes;
+		TArray<FEvalNodePiecewiseLinear> PiecewiseLinearNodes;
+		TArray<FEvalNodeMultiply> MultiplyNodes;
+	};
+
+	enum class ENodeType : uint8
+	{
+		None,
+		Input,
+		WeightedSum,
+		PiecewiseLinear,
+		Multiply
+	};
 
 	/**
-	 * 
+	 * Lighweight pointer/proxy to a node in the FNodeCollection.
+	 */
+	class FNodePtr
+	{
+	public:
+		FNodePtr() = default;
+		FNodePtr(ENodeType InNodeType, int32 InDataIndex);
+
+		bool TryGetValue(const FNodeCollection& NodeCollection, double& OutValue) const;
+		double GetValue(const FNodeCollection& NodeCollection) const;
+		void SetDirect(FNodeCollection& NodeCollection, double Value);
+		void Reset(FNodeCollection& NodeCollection);
+
+		bool IsValid() const
+		{
+			return NodeType != ENodeType::None;
+		}
+
+		bool IsInitialized(const FNodeCollection& NodeCollection) const;
+
+		friend bool operator==(const FNodePtr& Lhs, const FNodePtr& Rhs)
+		{
+			return Lhs.DataIndex == Rhs.DataIndex && Lhs.NodeType == Rhs.NodeType;
+		}
+
+		friend uint32 GetTypeHash(const FNodePtr& Node)
+		{
+			return HashCombineFast(GetTypeHash(Node.DataIndex), GetTypeHash(Node.NodeType));
+		}
+
+	private:
+		int32 DataIndex = INDEX_NONE;
+		ENodeType NodeType = ENodeType::None;
+	};
+
+	/**
+	 * Rig Mapper node representation loaded from URigMapperDefinition
 	 */
 	class RIGMAPPER_API FRigMapper
 	{
 	public:
 		bool LoadDefinition(const URigMapperDefinition* Definition);
 		bool SetDirectValue(int32 Index, double Value);
-		bool SetDirectValue(const FString& InputName, double Value);
+		bool SetDirectValue(const FName& InputName, double Value);
 		void Reset();
 		bool IsValid() const;
 
-		TMap<FString,double> GetOutputValues(bool bSkipUnset = false) const;
-		const TArray<FString>& GetOutputNames() const;
+		TMap<FName,double> GetOutputValues(bool bSkipUnset = false) const;
+		const TArray<FName>& GetOutputNames() const;
 		void GetOutputValuesInOrder(TArray<double>& OutValues) const;
 		void GetOptionalOutputValuesInOrder(TArray<TOptional<double>>& OutValues) const;
 		void GetOptionalFloatOutputValuesInOrder(TArray<TOptional<float>>& OutValues) const;
-		
-		void SetDirty() const;
 
-		const TArray<FString>& GetInputNames() const;
+		void SetDirty();
+
+		const TArray<FName>& GetInputNames() const;
 
 	private:
-		TMap<FString, EvalPtr> Nodes;
-		
-		TArray<EvalPtr> OutputNodes;
-		TArray<FString> OutputNames;
-		
-		TArray<EvalPtr> InputNodes;
-		TArray<FString> InputNames;
+		TMap<FName, FNodePtr> Nodes;
+		FNodeCollection NodeCollection;
+
+		TArray<FNodePtr> OutputNodes;
+		TArray<FName> OutputNames;
+
+		TArray<FNodePtr> InputNodes;
+		TArray<FName> InputNames;
 	};
 
-	/**
-	 *
-	 */
-	class FEvalNode
+
+	class FInputNode
 	{
 	public:
-		explicit FEvalNode(const FString& InName);
-		virtual ~FEvalNode(){}
-	
-		double GetValue() const;
-		bool TryGetValue(double &OutValue) const;
-		void SetDirect(double Value);
-		void Reset();
-
-		virtual void Initialize(const TMap<FString, EvalPtr>& Nodes) {}
+		bool TryGetValue(const FNodeCollection& NodeCollection, double& OutValue) const
+		{
+			if (CachedValue.IsSet())
+			{
+				OutValue = CachedValue.GetValue();
+				return true;
+			}
+			return false;
+		}
 		
-		bool IsInitialized() const;
-
-	protected:
-		bool bInitialized = false;
+		void SetDirect(double Value)
+		{
+			CachedValue = Value;
+		}
 		
-	private:
-		virtual bool Evaluate(double &OutValue) const { return false; }
+		void Reset()
+		{
+			CachedValue.Reset();
+		}
 
-		FString Name;
+	protected:
 		mutable TOptional<double> CachedValue;
 	};
-
-	/**
-	 *
-	 */
-	class FEvalInputNode : public FEvalNode
+	
+	template<typename DerivedNode>
+	class TEvalNode
 	{
 	public:
-		explicit FEvalInputNode(const FString& InName) : FEvalNode(InName) {}
-		
-		virtual void Initialize(const TMap<FString, EvalPtr>& Nodes) override { bInitialized = true; }
+		bool TryGetValue(const FNodeCollection& NodeCollection, double& OutValue) const
+		{
+			if (bInitialized)
+			{
+				if (!CachedValue.IsSet())
+				{
+					double Value;
+					if (static_cast<const DerivedNode*>(this)->Evaluate(NodeCollection, Value))
+					{
+						CachedValue = Value;
+						OutValue = Value;
+						return true;
+					}
+				}
+				else
+				{
+					OutValue = CachedValue.GetValue();
+					return true;
+				}
+			}
+			return false;
+		}
+
+		bool IsInitialized() const
+		{
+			return bInitialized;
+		}
+
+		void SetDirect(double Value)
+		{
+			CachedValue = Value;
+		}
+
+		void Reset()
+		{
+			CachedValue.Reset();
+		}
+
+	protected:
+		mutable TOptional<double> CachedValue;
+		bool bInitialized = false;
 	};
 
-	/**
-	 * 
-	 */
-	class FEvalNodeWeightedSum : public FEvalNode
+	class FEvalNodeWeightedSum : public TEvalNode<FEvalNodeWeightedSum>
 	{
 	public:
-		FEvalNodeWeightedSum(const FRigMapperWsFeature& FeatureDefinition);
+		void Initialize(const FRigMapperWsFeature& FeatureDefinition, const TMap<FName, FNodePtr>& Nodes);
+		bool Evaluate(const FNodeCollection& NodeCollection, double& OutValue) const;
 
 	private:
-		virtual bool Evaluate(double &OutValue) const override;
-
-		virtual void Initialize(const TMap<FString, EvalPtr>& Nodes) override;
-
-		TMap<FString, double> WeightedLinkedInputNames;
-	
-		TMap<EvalPtr, double> WeightedLinkedInputs;
-
+		TArray<TPair<FNodePtr, double>> WeightedLinkedInputs;
 		TRange<double> Range;
 	};
 
-	/**
-	 *
-	 */
-	class FEvalNodePiecewiseLinear :public FEvalNode
+	class FEvalNodePiecewiseLinear : public TEvalNode<FEvalNodePiecewiseLinear>
 	{
 	public:
-		FEvalNodePiecewiseLinear(const FRigMapperSdkFeature& FeatureDefinition);
+		void Initialize(const FRigMapperSdkFeature& FeatureDefinition, const TMap<FName, FNodePtr>& Nodes);
+		bool Evaluate(const FNodeCollection& NodeCollection, double& OutValue) const;
 
-	public:
-		static bool Evaluate_Static(const double InputValue, const TArray<TPair<double,double>>& Keys, double &OutValue);
-		
-	private:
-		virtual bool Evaluate(double &OutValue) const override;
-		
-		virtual void Initialize(const TMap<FString, EvalPtr>& Nodes) override;
+		static bool Evaluate_Static(const double InputValue, TConstArrayView<TPair<double, double>> Keys, double& OutValue);
 
 	private:
-		FString LinkedInputName;
-		EvalPtr LinkedInput;
-		TArray<TPair<double,double>> Keys;
+		FNodePtr LinkedInput;
+		// 95% only have 2 elements in this array.
+		TArray<TPair<double, double>, TInlineAllocator<2>> Keys;
 	};
 
-	/**
-	 * 
-	 */
-	class FEvalNodeMultiply : public FEvalNode
+	class FEvalNodeMultiply : public TEvalNode<FEvalNodeMultiply>
 	{
 	public:
-		explicit FEvalNodeMultiply(const FRigMapperMultiplyFeature& FeatureDefinition);
+		void Initialize(const FRigMapperMultiplyFeature& FeatureDefinition, const TMap<FName, FNodePtr>& Nodes);
+		bool Evaluate(const FNodeCollection& NodeCollection, double& OutValue) const;
 
 	private:
-		virtual bool Evaluate(double &OutValue) const override;
-
-		virtual void Initialize(const TMap<FString, EvalPtr>& Nodes) override;
-
-		TArray<FString> LinkedInputNames;
-		TArray<EvalPtr> LinkedInputs;
+		TArray<FNodePtr, TInlineAllocator<2>> LinkedInputs;
 	};
 }
diff --git a/Engine/Plugins/Experimental/Animation/RigMapper/Source/RigMapper/Public/RigMapperProcessor.h b/Engine/Plugins/Experimental/Animation/RigMapper/Source/RigMapper/Public/RigMapperProcessor.h
--- a/Engine/Plugins/Experimental/Animation/RigMapper/Source/RigMapper/Public/RigMapperProcessor.h
+++ b/Engine/Plugins/Experimental/Animation/RigMapper/Source/RigMapper/Public/RigMapperProcessor.h
@@ -7,6 +7,12 @@
 #include "Misc/FrameTime.h"
 #include "RigMapperProcessor.generated.h"
 
+DECLARE_STATS_GROUP(TEXT("RigMapper"), STATGROUP_RigMapper, STATCAT_Advanced);
+
+DECLARE_CYCLE_STAT_EXTERN(TEXT("Get Rig Mapper"), STAT_RigMapperGetRigMapper, STATGROUP_RigMapper, RIGMAPPER_API);
+DECLARE_CYCLE_STAT_EXTERN(TEXT("Load Definition"), STAT_RigMapperLoadDefinition, STATGROUP_RigMapper, RIGMAPPER_API);
+
+
 /**
  * A helper struct to gain performance when doing some batch remapping with the Rig Mapper pipeline
  * Useful when remapping frames every tick, or numerous frames at once, or when using chained definitions...
@@ -25,45 +31,40 @@ struct RIGMAPPER_API FRigMapperProcessor
 	FRigMapperProcessor() = default;
 	explicit FRigMapperProcessor(const TArray<URigMapperDefinition*>& InDefinitions);
 	explicit FRigMapperProcessor(URigMapperDefinition* InDefinition);
-	explicit FRigMapperProcessor(TSharedPtr<FacialRigMapping::FRigMapper> InRigMapper);
-	explicit FRigMapperProcessor(const TArray<TSharedPtr<FacialRigMapping::FRigMapper>>& InRigMappers);
-
-	/** Loads a set of rig mappers from a set of definition assets */
-	static bool LoadDefinitions(const TArray<URigMapperDefinition*>& Definitions, TArray<TSharedPtr<FacialRigMapping::FRigMapper>>& RigMappers);
 
 	/** @returns whether we have at least one rig mapper initialized, and all rig mappers are valid */
 	bool IsValid() const;
 
 	/** @returns the input names expected by the first rig mapper */
-	const TArray<FString>& GetInputNames() const;
+	const TArray<FName>& GetInputNames() const;
 
 	/** @returns the output names given by the last rig mapper */
-	const TArray<FString>& GetOutputNames() const;
+	const TArray<FName>& GetOutputNames() const;
 
 	/** Evaluates a set of frames all at once - retrieve and/or cache curve names using @GetOutputNames for performance */
-	bool EvaluateFrames(const TArray<FString>& CurveNames, const TArray<FPoseValues>& InFrameValues, TArray<FPoseValues>& OutFrameValues);
+	bool EvaluateFrames(const TArray<FName>& CurveNames, const TArray<FPoseValues>& InFrameValues, TArray<FPoseValues>& OutFrameValues);
 
 	/** Evaluates a set of frames all at once */
-	bool EvaluateFrames(const TArray<FString>& CurveNames, const TArray<FPoseValues>& InFrameValues, TArray<FString>& OutCurveNames, TArray<FPoseValues>& OutFrameValues);
+	bool EvaluateFrames(const TArray<FName>& CurveNames, const TArray<FPoseValues>& InFrameValues, TArray<FName>& OutCurveNames, TArray<FPoseValues>& OutFrameValues);
 
 	/** Evaluates a set of frames all at once, interpolating between them using the given frame times */
-	bool EvaluateFrames_Interp(const TArray<FString>& CurveNames, const TArray<FPoseValues>& InFrameValues, TArray<FPoseValues>& OutFrameValues, const TArray<FFrameTime>& FrameTimes);
+	bool EvaluateFrames_Interp(const TArray<FName>& CurveNames, const TArray<FPoseValues>& InFrameValues, TArray<FPoseValues>& OutFrameValues, const TArray<FFrameTime>& FrameTimes);
 
 	/** Evaluates a single frame - retrieve and/or cache curve names using @GetOutputNames for performance */
-	bool EvaluateFrame(const TArray<FString>& CurveNames, const FPoseValues& InCurveValues, FPoseValues& OutCurveValues);
+	bool EvaluateFrame(const TArray<FName>& CurveNames, const FPoseValues& InCurveValues, FPoseValues& OutCurveValues);
 
 	/** Evaluates a single frame */
-	bool EvaluateFrame(const TArray<FString>& CurveNames, const FPoseValues& InCurveValues, TArray<FString>& OutCurveNames, FPoseValues& OutCurveValues);
+	bool EvaluateFrame(const TArray<FName>& CurveNames, const FPoseValues& InCurveValues, TArray<FName>& OutCurveNames, FPoseValues& OutCurveValues);
 
 private:
 	/** Evaluates a single frame for a single rig mapper */
-	bool EvaluateFrame_Internal(const TArray<FString>& CurveNames, const FPoseValues& InCurveValues, FPoseValues& OutCurveValues, int32 RigIndex);
+	bool EvaluateFrame_Internal(const TArray<FName>& CurveNames, const FPoseValues& InCurveValues, FPoseValues& OutCurveValues, int32 RigIndex);
 	
 private:
 	using FIndexCache = TArray<int32>;
 
 	/** The rig mappers loaded from definitions */
-	TArray<TSharedPtr<FacialRigMapping::FRigMapper>> RigMappers;
+	TArray<FacialRigMapping::FRigMapper> RigMappers;
 
 	/** An index map for each rig mapper's input - if the input curve names change between two calls, it will be rebuilt */
 	TArray<FIndexCache> IndexCache;
diff --git a/Engine/Plugins/Experimental/Animation/RigMapper/Source/RigMapperDeveloper/Private/Tests/RigMapperTests.cpp b/Engine/Plugins/Experimental/Animation/RigMapper/Source/RigMapperDeveloper/Private/Tests/RigMapperTests.cpp
--- a/Engine/Plugins/Experimental/Animation/RigMapper/Source/RigMapperDeveloper/Private/Tests/RigMapperTests.cpp
+++ b/Engine/Plugins/Experimental/Animation/RigMapper/Source/RigMapperDeveloper/Private/Tests/RigMapperTests.cpp
@@ -323,14 +323,14 @@ using FRigMapperIndexMap = TArray<int32>;
 
 FRigMapperIndexMap MakeIndexMap(const TSharedPtr<FacialRigMapping::FRigMapper>& InOutputRig, const TSharedPtr<FacialRigMapping::FRigMapper>& InInputRig)
 {
-	const TArray<FString> FromOutputs = InOutputRig->GetOutputNames();
-	const TArray<FString> ToInputs =  InInputRig->GetInputNames();
+	const TArray<FName> FromOutputs = InOutputRig->GetOutputNames();
+	const TArray<FName> ToInputs =  InInputRig->GetInputNames();
 
 	FRigMapperIndexMap IndexMapping;
 
 	IndexMapping.Reserve(ToInputs.Num());
 
-	for (const FString& Input : ToInputs)
+	for (const FName& Input : ToInputs)
 	{
 		IndexMapping.Add(FromOutputs.Find(Input));
 	}
@@ -395,12 +395,12 @@ void EvaluateRigMapping(const TMap<FString, double>& InInputs, const TArray<FRig
 		if (i == 0)
 		{
 
-			const TArray<FString>& InputNames = RigMapper->GetInputNames();
+			const TArray<FName>& InputNames = RigMapper->GetInputNames();
 			for (int32 n = 0; n < InputNames.Num(); n++)
 			{
-				if (InInputs.Find(*InputNames[n]))
+				if (const double* FoundInput = InInputs.Find(InputNames[n].ToString()); FoundInput != nullptr)
 				{
-					RigMapper->SetDirectValue(n, InInputs[*InputNames[n]]);
+					RigMapper->SetDirectValue(n, *FoundInput);
 				}
 			}
 		}
@@ -432,7 +432,11 @@ void EvaluateRigMapping(const TMap<FString, double>& InInputs, const TArray<FRig
 	// Set outputs from last rig mapper's output values 
 	if (!InOutRigMappers.IsEmpty())
 	{
-		OutOutputs = InOutRigMappers.Last()->GetOutputValues();
+		OutOutputs.Empty();
+		for (const TPair<FName, double> OutputPair : InOutRigMappers.Last()->GetOutputValues())
+		{
+			OutOutputs.Add(OutputPair.Key.ToString(), OutputPair.Value);
+		}
 	}
 }
 
@@ -561,7 +565,7 @@ bool FRigMapperProcessorTest::RunTest(const FString& InParameters)
 	}
 
 	// set up some sample inputs
-	TArray<FString> InputNames;
+	TArray<FName> InputNames;
 	InputNames.Add(TEXT("InputVal1"));
 	InputNames.Add(TEXT("InputVal2"));
 	InputNames.Add(TEXT("InputVal3"));
@@ -582,7 +586,7 @@ bool FRigMapperProcessorTest::RunTest(const FString& InParameters)
 		UE_LOG(LogRigMapperTest, Error, TEXT("RigMapperProcessor failed to evaluate input values"));
 		return false;
 	}
-	const TArray<FString>& OutputNames = Processor.GetOutputNames();
+	const TArray<FName>& OutputNames = Processor.GetOutputNames();
 
 	// check that the outputs are correct
 
@@ -619,7 +623,7 @@ bool FRigMapperProcessorTest::RunTest(const FString& InParameters)
 	//		OutputVal10 = TestMultFeature3 = 0.096
 	//		OutputVal11 = TestSDKFeature3 = 0.35
 
-	int32 OutputIndex = OutputNames.Find(FString(TEXT("OutputVal9")));
+	int32 OutputIndex = OutputNames.Find(FName(TEXT("OutputVal9")));
 	if (!OutputValues.IsValidIndex(OutputIndex))
 	{
 		UE_LOG(LogRigMapperTest, Error, TEXT("Invalid index for Output Value 9 (TestWSFeature3): %d"), OutputIndex);
@@ -636,7 +640,7 @@ bool FRigMapperProcessorTest::RunTest(const FString& InParameters)
 		return false;
 	}
 
-	OutputIndex = OutputNames.Find(FString(TEXT("OutputVal10")));
+	OutputIndex = OutputNames.Find(FName(TEXT("OutputVal10")));
 	if (!OutputValues.IsValidIndex(OutputIndex))
 	{
 		UE_LOG(LogRigMapperTest, Error, TEXT("Invalid index for Output Value 10 (TestMultFeature3): %d"), OutputIndex);
@@ -653,7 +657,7 @@ bool FRigMapperProcessorTest::RunTest(const FString& InParameters)
 		return false;
 	}
 
-	OutputIndex = OutputNames.Find(FString(TEXT("OutputVal11")));
+	OutputIndex = OutputNames.Find(FName(TEXT("OutputVal11")));
 	if (!OutputValues.IsValidIndex(OutputIndex))
 	{
 		UE_LOG(LogRigMapperTest, Error, TEXT("Invalid index for Output Value 11 (TestSdkFeature3): %d"), OutputIndex);
diff --git a/Engine/Plugins/Experimental/Animation/RigMapper/Source/RigMapperEditor/Private/RigMapperEditorSubsystem.cpp b/Engine/Plugins/Experimental/Animation/RigMapper/Source/RigMapperEditor/Private/RigMapperEditorSubsystem.cpp
--- a/Engine/Plugins/Experimental/Animation/RigMapper/Source/RigMapperEditor/Private/RigMapperEditorSubsystem.cpp
+++ b/Engine/Plugins/Experimental/Animation/RigMapper/Source/RigMapperEditor/Private/RigMapperEditorSubsystem.cpp
@@ -197,8 +197,14 @@ bool URigMapperEditorSubsystem::ConvertCurveValuesToCsv(TArray<FString>& CurveNa
 
 	// todo: bOutputIntermediateCsvFiles
 	// const FString BaseFilePath = FPaths::GetBaseFilename(OutputFile.FilePath, false);;
-	
-	RigMapperProcessor.EvaluateFrames(CurveNames, CurveValuesPerFrame, CurveNames, CurveValuesPerFrame);
+	TArray<FName> CurveNamesConverted;
+	CurveNamesConverted.Reserve(CurveNames.Num());
+	for (const FString& Str : CurveNames)
+	{
+		CurveNamesConverted.Add(FName(*Str));
+	}
+
+	RigMapperProcessor.EvaluateFrames(CurveNamesConverted, CurveValuesPerFrame, CurveNamesConverted, CurveValuesPerFrame);
 
 	WriteCurveValuesToCsv(OutputFile, CurveNames, FrameTimes, CurveValuesPerFrame);
 	
@@ -221,6 +227,14 @@ bool URigMapperEditorSubsystem::ConvertCsv(const FFilePath& InputFile, const FFi
 
 bool URigMapperEditorSubsystem::ConvertCurveValuesToAnimSequence(TArray<FString>& CurveNames, const TArray<FFrameTime>& FrameTimes, FFrameValues& CurveValuesPerFrame, UAnimSequence* Target, const TArray<URigMapperDefinition*>& Definitions)
 {
+	TArray<FName> CurveNamesConverted;
+	CurveNamesConverted.Reserve(CurveNames.Num());
+	for (const FString& Str : CurveNames)
+	{
+		CurveNamesConverted.Add(FName(*Str));
+	}
+
+
 	FRigMapperProcessor RigMapperProcessor(Definitions);
 	if (!RigMapperProcessor.IsValid())
 	{
@@ -228,7 +242,7 @@ bool URigMapperEditorSubsystem::ConvertCurveValuesToAnimSequence(TArray<FString>
 		return false;
 	}
 	FFrameValues ConvertedCurveValuesPerFrame;
-	if (!RigMapperProcessor.EvaluateFrames(CurveNames, CurveValuesPerFrame, CurveNames, ConvertedCurveValuesPerFrame))
+	if (!RigMapperProcessor.EvaluateFrames(CurveNamesConverted, CurveValuesPerFrame, CurveNamesConverted, ConvertedCurveValuesPerFrame))
 	{
 		UE_LOG(LogRigMapperEditor, Error, TEXT("Rig Mapper Processor failed to evaluate frames"))
 		return false;
@@ -453,12 +467,19 @@ UMovieSceneControlRigParameterSection* URigMapperEditorSubsystem::ConvertAnimSeq
 
 bool URigMapperEditorSubsystem::ConvertCurveValuesToControlRigSection(TArray<FString>& CurveNames, const TArray<FFrameTime>& FrameTimes, FFrameValues& CurveValuesPerFrame, UMovieSceneControlRigParameterSection* Target, const TArray<URigMapperDefinition*>& Definitions)
 {
+	TArray<FName> CurveNamesConverted;
+	CurveNamesConverted.Reserve(CurveNames.Num());
+	for (const FString& Str : CurveNames)
+	{
+		CurveNamesConverted.Add(FName(*Str));
+	}
+
 	FRigMapperProcessor RigMapperProcessor(Definitions);
 	if (!RigMapperProcessor.IsValid())
 	{
 		return false;
 	}
-	RigMapperProcessor.EvaluateFrames(CurveNames, CurveValuesPerFrame, CurveNames, CurveValuesPerFrame);
+	RigMapperProcessor.EvaluateFrames(CurveNamesConverted, CurveValuesPerFrame, CurveNamesConverted, CurveValuesPerFrame);
 
 	return AddCurveValuesToControlRigSection(Target, CurveNames, FrameTimes, CurveValuesPerFrame);
 }
