[RigMapper] fix bug so that if RigMapperDefinition is edited, changes are reflected in any RigMapper node using it.

diff --git a/Engine/Plugins/Experimental/Animation/RigMapper/Source/RigMapper/Private/RigMapperDefinition.cpp b/Engine/Plugins/Experimental/Animation/RigMapper/Source/RigMapper/Private/RigMapperDefinition.cpp
--- a/Engine/Plugins/Experimental/Animation/RigMapper/Source/RigMapper/Private/RigMapperDefinition.cpp
+++ b/Engine/Plugins/Experimental/Animation/RigMapper/Source/RigMapper/Private/RigMapperDefinition.cpp
@@ -3,6 +3,7 @@
 #include "RigMapperDefinition.h"
 
 #include "RigMapper.h"
+#include "RigMapperProcessor.h"
 #include "Misc/FileHelper.h"
 #include "Serialization/JsonTypes.h"
 #include "Dom/JsonValue.h"
@@ -773,6 +774,10 @@ bool URigMapperDefinition::Validate()
 void URigMapperDefinition::PostEditChangeProperty(FPropertyChangedEvent& PropertyChangedEvent)
 {
 	Super::PostEditChangeProperty(PropertyChangedEvent);
+
+	// if the definition exists in the global cache, delete it so that the editor doesn't get out of sync with the cache
+	FRigMapperDefinitionsSingleton::Get().ClearFromCache(this);
+
 	OnRigMapperDefinitionUpdated.Broadcast();
 
 	SetDefinitionValid(false);
diff --git a/Engine/Plugins/Experimental/Animation/RigMapper/Source/RigMapper/Private/RigMapperProcessor.cpp b/Engine/Plugins/Experimental/Animation/RigMapper/Source/RigMapper/Private/RigMapperProcessor.cpp
--- a/Engine/Plugins/Experimental/Animation/RigMapper/Source/RigMapper/Private/RigMapperProcessor.cpp
+++ b/Engine/Plugins/Experimental/Animation/RigMapper/Source/RigMapper/Private/RigMapperProcessor.cpp
@@ -8,21 +8,18 @@ DEFINE_STAT(STAT_RigMapperGetRigMapper);
 DEFINE_STAT(STAT_RigMapperLoadDefinition);
 
 
-class FRigMapperDefinitionsSingleton
-{
-public:
-	static FRigMapperDefinitionsSingleton& Get();
-	bool GetRigMapper(const TObjectPtr<const URigMapperDefinition>& InDefinition, FacialRigMapping::FRigMapper& OutRigMapper);
-private:
-	FRigMapperDefinitionsSingleton() = default;
-	virtual ~FRigMapperDefinitionsSingleton() = default;
-	TMap<FSoftObjectPath, FacialRigMapping::FRigMapper> RigMappers;
-};
 FRigMapperDefinitionsSingleton& FRigMapperDefinitionsSingleton::Get()
 {
 	static FRigMapperDefinitionsSingleton SingletonInstance;
 	return SingletonInstance;
 }
+
+void FRigMapperDefinitionsSingleton::ClearFromCache(const TObjectPtr<const URigMapperDefinition>& InDefinition)
+{
+	// remove the rig mapper from the cache if it exists
+	RigMappers.Remove(InDefinition);
+}
+
 bool FRigMapperDefinitionsSingleton::GetRigMapper(const TObjectPtr<const URigMapperDefinition>& InDefinition, FacialRigMapping::FRigMapper& OutRigMapper)
 {
 	SCOPE_CYCLE_COUNTER(STAT_RigMapperGetRigMapper);
diff --git a/Engine/Plugins/Experimental/Animation/RigMapper/Source/RigMapper/Public/RigMapperProcessor.h b/Engine/Plugins/Experimental/Animation/RigMapper/Source/RigMapper/Public/RigMapperProcessor.h
--- a/Engine/Plugins/Experimental/Animation/RigMapper/Source/RigMapper/Public/RigMapperProcessor.h
+++ b/Engine/Plugins/Experimental/Animation/RigMapper/Source/RigMapper/Public/RigMapperProcessor.h
@@ -72,4 +72,21 @@ private:
 	TArray<FIndexCache> IndexCache;
 };
 
+/*
+* A singleton class which stores a cache of FRigMappers initialized from URigMapperDefinitions, for speed
+*/
+class FRigMapperDefinitionsSingleton
+{
+public:
+	static FRigMapperDefinitionsSingleton& Get();
+	bool GetRigMapper(const TObjectPtr<const URigMapperDefinition>& InDefinition, FacialRigMapping::FRigMapper& OutRigMapper);
+
+	/* clear rig mapper from the cache if it exists*/
+	void ClearFromCache(const TObjectPtr<const URigMapperDefinition>& InDefinition);
+private:
+	FRigMapperDefinitionsSingleton() = default;
+	virtual ~FRigMapperDefinitionsSingleton() = default;
+	TMap<FSoftObjectPath, FacialRigMapping::FRigMapper> RigMappers;
+};
+
 #undef UE_API
