[RigMapper] fixed crash when create a RigMapper node and add definition

diff --git a/Engine/Plugins/Experimental/Animation/RigMapper/Source/RigMapper/Private/RigMapperProcessor.cpp b/Engine/Plugins/Experimental/Animation/RigMapper/Source/RigMapper/Private/RigMapperProcessor.cpp
--- a/Engine/Plugins/Experimental/Animation/RigMapper/Source/RigMapper/Private/RigMapperProcessor.cpp
+++ b/Engine/Plugins/Experimental/Animation/RigMapper/Source/RigMapper/Private/RigMapperProcessor.cpp
@@ -50,15 +50,18 @@ FRigMapperProcessor::FRigMapperProcessor(const TArray<URigMapperDefinition*>& In
 {
 	for (const URigMapperDefinition* Definition : InDefinitions)
 	{
-		FacialRigMapping::FRigMapper RigMapper;
-		if (!FRigMapperDefinitionsSingleton::Get().GetRigMapper(Definition, RigMapper))
-
+		if (Definition)
 		{
-			UE_LOG(LogRigMapper, Error, TEXT("Could not load definition %s"), *Definition->GetPathName())
-				RigMappers.Empty();
-			return;
+			FacialRigMapping::FRigMapper RigMapper;
+			if (!FRigMapperDefinitionsSingleton::Get().GetRigMapper(Definition, RigMapper))
+
+			{
+				UE_LOG(LogRigMapper, Error, TEXT("Could not load definition %s"), *Definition->GetPathName())
+					RigMappers.Empty();
+				return;
+			}
+			RigMappers.Add(MoveTemp(RigMapper));
 		}
-		RigMappers.Add(MoveTemp(RigMapper));
 	}
 
 	IndexCache.AddDefaulted(RigMappers.Num());
@@ -66,16 +69,19 @@ FRigMapperProcessor::FRigMapperProcessor(const TArray<URigMapperDefinition*>& In
 
 FRigMapperProcessor::FRigMapperProcessor(URigMapperDefinition* InDefinition)
 {
-	FacialRigMapping::FRigMapper RigMapper;
-	if (!FRigMapperDefinitionsSingleton::Get().GetRigMapper(InDefinition, RigMapper))
-
+	if (InDefinition)
 	{
-		UE_LOG(LogRigMapper, Error, TEXT("Could not load definition %s"), *InDefinition->GetPathName())
-			RigMappers.Empty();
-		return;
+		FacialRigMapping::FRigMapper RigMapper;
+		if (!FRigMapperDefinitionsSingleton::Get().GetRigMapper(InDefinition, RigMapper))
+
+		{
+			UE_LOG(LogRigMapper, Error, TEXT("Could not load definition %s"), *InDefinition->GetPathName())
+				RigMappers.Empty();
+			return;
+		}
+		RigMappers.Add(MoveTemp(RigMapper));
+		IndexCache.AddDefaulted(RigMappers.Num());
 	}
-	RigMappers.Add(MoveTemp(RigMapper));
-	IndexCache.AddDefaulted(RigMappers.Num());
 }
 
 bool FRigMapperProcessor::IsValid() const
