[RigMapper] added a EUB to allow easy conversion of Anim Sequences using RigMapper Definitions

#virtualized

diff --git a/Engine/Plugins/Experimental/Animation/RigMapper/Source/RigMapperEditor/Private/RigMapperEditorSubsystem.cpp b/Engine/Plugins/Experimental/Animation/RigMapper/Source/RigMapperEditor/Private/RigMapperEditorSubsystem.cpp
--- a/Engine/Plugins/Experimental/Animation/RigMapper/Source/RigMapperEditor/Private/RigMapperEditorSubsystem.cpp
+++ b/Engine/Plugins/Experimental/Animation/RigMapper/Source/RigMapperEditor/Private/RigMapperEditorSubsystem.cpp
@@ -22,6 +22,11 @@
 
 UAnimSequence* CreateNewAnimSequence(const FString& BaseName, const FString& Path, USkeletalMesh* TargetMesh)
 {
+	if (!TargetMesh)
+	{
+		return nullptr;
+	}
+
 	FString UniqueAssetName = BaseName;
 	
 	FString UniquePackageName;
@@ -206,6 +211,13 @@ bool URigMapperEditorSubsystem::ConvertCurveValuesToCsv(TArray<FString>& CurveNa
 
 	RigMapperProcessor.EvaluateFrames(CurveNamesConverted, CurveValuesPerFrame, CurveNamesConverted, CurveValuesPerFrame);
 
+	CurveNames.Empty();
+	CurveNames.Reserve(CurveNamesConverted.Num()); 
+	Algo::Transform(CurveNamesConverted, CurveNames, [](const FName& Name)
+		{
+			return Name.ToString();
+		});
+
 	WriteCurveValuesToCsv(OutputFile, CurveNames, FrameTimes, CurveValuesPerFrame);
 	
 	return true;
@@ -247,6 +259,15 @@ bool URigMapperEditorSubsystem::ConvertCurveValuesToAnimSequence(TArray<FString>
 		UE_LOG(LogRigMapperEditor, Error, TEXT("Rig Mapper Processor failed to evaluate frames"))
 		return false;
 	}
+
+	CurveNames.Empty();
+	CurveNames.Reserve(CurveNamesConverted.Num());
+	Algo::Transform(CurveNamesConverted, CurveNames, [](const FName& Name)
+		{
+			return Name.ToString();
+		});
+
+
 	return AddCurveValuesToAnimSequence(Target, CurveNames, FrameTimes, ConvertedCurveValuesPerFrame);
 }
 
@@ -377,15 +398,24 @@ UAnimSequence* URigMapperEditorSubsystem::ConvertAnimSequenceNew(const UAnimSequ
 {
 	UAnimSequence* NewSequence = CreateNewAnimSequence(NewAssetName.ToString(), NewAssetPath.Path, TargetMesh);
 
-	NewSequence->ImportFileFramerate = Source->ImportFileFramerate;
-	NewSequence->ImportResampleFramerate = Source->ImportResampleFramerate;
-	SetAnimSequenceRate(NewSequence, GetAnimSequenceRate(Source), false);
-	
-	if (!ConvertAnimSequence(Source, NewSequence, Definitions))
+	if (NewSequence)
 	{
-		return nullptr;
+		NewSequence->ImportFileFramerate = Source->ImportFileFramerate;
+		NewSequence->ImportResampleFramerate = Source->ImportResampleFramerate;
+		SetAnimSequenceRate(NewSequence, GetAnimSequenceRate(Source), false);
+
+		if (!ConvertAnimSequence(Source, NewSequence, Definitions))
+		{
+			UE_LOG(LogRigMapperEditor, Error, TEXT("Failed to convert anim sequence using selected RigMapper Definitions"));
+			return nullptr;
+		}
+
+		NewSequence->PostLoad();
+	}
+	else
+	{
+		UE_LOG(LogRigMapperEditor, Error, TEXT("Failed to create anim sequence %s in folder %s"), *NewAssetName.ToString(), *NewAssetPath.Path);
 	}
-	NewSequence->PostLoad();
 	return NewSequence;
 }
 
@@ -481,6 +511,14 @@ bool URigMapperEditorSubsystem::ConvertCurveValuesToControlRigSection(TArray<FSt
 	}
 	RigMapperProcessor.EvaluateFrames(CurveNamesConverted, CurveValuesPerFrame, CurveNamesConverted, CurveValuesPerFrame);
 
+
+	CurveNames.Empty();
+	CurveNames.Reserve(CurveNamesConverted.Num());
+	Algo::Transform(CurveNamesConverted, CurveNames, [](const FName& Name)
+		{
+			return Name.ToString();
+		});
+
 	return AddCurveValuesToControlRigSection(Target, CurveNames, FrameTimes, CurveValuesPerFrame);
 }
 
