[RigMapper] a couple of fixes in the Subsystem functions

diff --git a/Engine/Plugins/Experimental/Animation/RigMapper/Source/RigMapperEditor/Private/RigMapperEditorSubsystem.cpp b/Engine/Plugins/Experimental/Animation/RigMapper/Source/RigMapperEditor/Private/RigMapperEditorSubsystem.cpp
--- a/Engine/Plugins/Experimental/Animation/RigMapper/Source/RigMapperEditor/Private/RigMapperEditorSubsystem.cpp
+++ b/Engine/Plugins/Experimental/Animation/RigMapper/Source/RigMapperEditor/Private/RigMapperEditorSubsystem.cpp
@@ -209,7 +209,8 @@ bool URigMapperEditorSubsystem::ConvertCurveValuesToCsv(TArray<FString>& CurveNa
 		CurveNamesConverted.Add(FName(*Str));
 	}
 
-	RigMapperProcessor.EvaluateFrames(CurveNamesConverted, CurveValuesPerFrame, CurveNamesConverted, CurveValuesPerFrame);
+	FFrameValues InputCurveValuesPerFrame = CurveValuesPerFrame; // take a copy as call to EvaluateFrames resets the output before evaluating the input
+	RigMapperProcessor.EvaluateFrames(CurveNamesConverted, InputCurveValuesPerFrame, CurveNamesConverted, CurveValuesPerFrame);
 
 	CurveNames.Empty();
 	CurveNames.Reserve(CurveNamesConverted.Num()); 
@@ -509,7 +510,9 @@ bool URigMapperEditorSubsystem::ConvertCurveValuesToControlRigSection(TArray<FSt
 	{
 		return false;
 	}
-	RigMapperProcessor.EvaluateFrames(CurveNamesConverted, CurveValuesPerFrame, CurveNamesConverted, CurveValuesPerFrame);
+
+	FFrameValues InputCurveValuesPerFrame = CurveValuesPerFrame; // take a copy as call to EvaluateFrames resets the output before evaluating the input
+	RigMapperProcessor.EvaluateFrames(CurveNamesConverted, InputCurveValuesPerFrame, CurveNamesConverted, CurveValuesPerFrame);
 
 
 	CurveNames.Empty();
