[RigMapper] Reimplemented RigMapper Definition Graph to be interactive:
- Nodes can be moved around
- Nodes can be created from a connection or standalone
- Nodes can be updated through details panel
- Context menu on right click on node or graph
- Nodes can be relinked and links can be removed

diff --git a/Engine/Plugins/Experimental/Animation/RigMapper/Source/RigMapper/Private/RigMapperDefinition.cpp b/Engine/Plugins/Experimental/Animation/RigMapper/Source/RigMapper/Private/RigMapperDefinition.cpp
--- a/Engine/Plugins/Experimental/Animation/RigMapper/Source/RigMapper/Private/RigMapperDefinition.cpp
+++ b/Engine/Plugins/Experimental/Animation/RigMapper/Source/RigMapper/Private/RigMapperDefinition.cpp
@@ -25,30 +25,32 @@ bool FRigMapperFeature::IsValid(const TArray<FString>& InputNames, bool bWarn) c
 		bValid = false;
 	}
 
-	TArray<FString> FeatureInputs;
-	GetInputs(FeatureInputs);
-	
-	if (FeatureInputs.Contains(Name))
+	if (GetFeatureType() != ERigMapperFeatureType::Input && GetFeatureType() != ERigMapperFeatureType::NullOutput)
 	{
-		UE_LOG(LogRigMapper, Warning, TEXT("Feature %s is referencing itself"), *Name)
-		bValid = false;
-	}
+		TArray<FString> FeatureInputs;
+		GetInputs(FeatureInputs);
 
-	if (FeatureInputs.IsEmpty())
-	{
-		UE_LOG(LogRigMapper, Warning, TEXT("Feature %s does not reference any input"), *Name)
-		bValid = false;
-	}
-	
-	for (const FString& Input : FeatureInputs)
-	{
-		if (!InputNames.Contains(Input))
+		if (FeatureInputs.Contains(Name))
 		{
-			UE_LOG(LogRigMapper, Warning, TEXT("Undefined input or feature %s referenced in feature %s"), *Input, *Name)
-			bValid = false;
+			UE_LOG(LogRigMapper, Warning, TEXT("Feature %s is referencing itself"), *Name)
+				bValid = false;
+		}
+
+		if (FeatureInputs.IsEmpty())
+		{
+			UE_LOG(LogRigMapper, Warning, TEXT("Feature %s does not reference any input"), *Name)
+				bValid = false;
+		}
+
+		for (const FString& Input : FeatureInputs)
+		{
+			if (!InputNames.Contains(Input))
+			{
+				UE_LOG(LogRigMapper, Warning, TEXT("Undefined input or feature %s referenced in feature %s"), *Input, *Name)
+					bValid = false;
+			}
 		}
 	}
-	
 	return bValid;
 }
 
@@ -822,7 +824,7 @@ bool URigMapperLinkedDefinitions::GetBakedInputRec(const FString& InputName, con
 
 bool FRigMapperMultiplyFeature::BakeInput(URigMapperLinkedDefinitions* LinkedDefinitions, int32 DefinitionIndex, FBakedInput& OutBakedInput) const
 {
-	// todo: refactor all bakeinput using baseclass
+	// todo: refactor all bake input using base class
 	const FString BakedInputName = FString::Printf(TEXT("%s:%d"), *Name, DefinitionIndex);
 	
 	const TSharedPtr<FRigMapperMultiplyFeature> BakedMultFeature = MakeShared<FRigMapperMultiplyFeature>(BakedInputName);
@@ -876,7 +878,7 @@ bool FRigMapperWsFeature::BakeInput(URigMapperLinkedDefinitions* LinkedDefinitio
 
 			for (const TPair<FString, double>& SubFeatureSubInput : SubFeatureWsBakedInput->Inputs)
 			{
-				// check for the existance of the SubFeatureSubInput.Key in the map already as we can get 'diamond' structures
+				// check for the existence of the SubFeatureSubInput.Key in the map already as we can get 'diamond' structures
 				if (BakedWsFeature->Inputs.Contains(SubFeatureSubInput.Key))
 				{
 					BakedWsFeature->Inputs[SubFeatureSubInput.Key] += SubFeatureSubInput.Value * FeatureInput.Value;
diff --git a/Engine/Plugins/Experimental/Animation/RigMapper/Source/RigMapper/Public/RigMapperDefinition.h b/Engine/Plugins/Experimental/Animation/RigMapper/Source/RigMapper/Public/RigMapperDefinition.h
--- a/Engine/Plugins/Experimental/Animation/RigMapper/Source/RigMapper/Public/RigMapperDefinition.h
+++ b/Engine/Plugins/Experimental/Animation/RigMapper/Source/RigMapper/Public/RigMapperDefinition.h
@@ -21,7 +21,9 @@ enum class ERigMapperFeatureType : uint8
 	Input,
 	WeightedSum,
 	SDK,
-	Multiply
+	Multiply,
+	Output,
+	NullOutput
 };
 
 USTRUCT(BlueprintType)
@@ -205,7 +207,7 @@ public:
 	UE_API virtual bool IsValid(const TArray<FString>& InputNames, bool bWarn = false) const override;
 
 	virtual ERigMapperFeatureType GetFeatureType() const override { return ERigMapperFeatureType::SDK; }
-	
+
 	virtual void GetInputs(TArray<FString>& OutInputs) const override { OutInputs = { Input }; };
 
 	UE_API virtual bool BakeInput(URigMapperLinkedDefinitions* LinkedDefinitions, int32 DefinitionIndex, FBakedInput& OutBakedInput) const override;
@@ -218,6 +220,43 @@ public:
 	}
 };
 
+USTRUCT(BlueprintType)
+struct FRigMapperOutputFeature : public FRigMapperFeature
+{
+	GENERATED_BODY()
+
+public:
+	UPROPERTY(EditAnywhere, BlueprintReadWrite, Category="Animation|Rig Mapper", meta=(DisplayPriority=1))
+	FString Input;
+
+public:
+	FRigMapperOutputFeature() {}
+
+	explicit FRigMapperOutputFeature(const FString& InName) : FRigMapperFeature(InName) { }
+
+	virtual ERigMapperFeatureType GetFeatureType() const override { return ERigMapperFeatureType::Output; }
+
+	virtual void GetInputs(TArray<FString>& OutInputs) const override { OutInputs = { Input }; };
+
+	bool operator==(const FRigMapperOutputFeature& InOther) const
+	{
+		return FRigMapperFeature::operator==(InOther) && Input == InOther.Input;
+	}
+};
+
+USTRUCT(BlueprintType)
+struct FRigMapperNullOutputFeature : public FRigMapperFeature
+{
+	GENERATED_BODY()
+
+public:
+	FRigMapperNullOutputFeature() {}
+
+	explicit FRigMapperNullOutputFeature(const FString& InName) : FRigMapperFeature(InName) { }
+
+	virtual ERigMapperFeatureType GetFeatureType() const override { return ERigMapperFeatureType::NullOutput; }
+};
+
 USTRUCT(BlueprintType)
 struct FRigMapperFeatureDefinitions
 {
diff --git a/Engine/Plugins/Experimental/Animation/RigMapper/Source/RigMapperEditor/Private/RigMapperDefinitionEditorToolkit.cpp b/Engine/Plugins/Experimental/Animation/RigMapper/Source/RigMapperEditor/Private/RigMapperDefinitionEditorToolkit.cpp
--- a/Engine/Plugins/Experimental/Animation/RigMapper/Source/RigMapperEditor/Private/RigMapperDefinitionEditorToolkit.cpp
+++ b/Engine/Plugins/Experimental/Animation/RigMapper/Source/RigMapperEditor/Private/RigMapperDefinitionEditorToolkit.cpp
@@ -155,7 +155,7 @@ void FRigMapperDefinitionEditorToolkit::HandleRigMapperDefinitionLoaded()
 {
 	if (GraphEditor)
 	{
-		GraphEditor->RebuildGraph();
+		//GraphEditor->RebuildGraph();
 	}
 	if (StructureView)
 	{
@@ -182,15 +182,30 @@ TSharedRef<SDockTab> FRigMapperDefinitionEditorToolkit::SpawnDetailsTab(const FS
 {
 	FPropertyEditorModule& PropertyModule = FModuleManager::LoadModuleChecked<FPropertyEditorModule>("PropertyEditor");
 	
-	FDetailsViewArgs Args;
-	Args.bHideSelectionTip = true;
-	Args.NameAreaSettings = FDetailsViewArgs::HideNameArea;
-
-	DetailsView = PropertyModule.CreateDetailView(Args);
-	DetailsView->RegisterInstancedCustomPropertyLayout(URigMapperDefinition::StaticClass(), FOnGetDetailCustomizationInstance::CreateStatic(&FDetailsViewCustomization::MakeInstance));
-	DetailsView->SetIsPropertyVisibleDelegate(FIsPropertyVisible::CreateRaw(this, &FRigMapperDefinitionEditorToolkit::HandleIsPropertyVisible));
-	DetailsView->OnFinishedChangingProperties().AddRaw(this, &FRigMapperDefinitionEditorToolkit::HandleFinishedChangingProperties);
-	DetailsView->SetObject(Definition);
+	FDetailsViewArgs DetailsViewArgs;
+	{
+		DetailsViewArgs.bAllowSearch = false;
+		DetailsViewArgs.NameAreaSettings = FDetailsViewArgs::HideNameArea;
+		DetailsViewArgs.bHideSelectionTip = true;
+		DetailsViewArgs.DefaultsOnlyVisibility = EEditDefaultsOnlyNodeVisibility::Automatic;
+		DetailsViewArgs.bShowOptions = false;
+		DetailsViewArgs.bAllowMultipleTopLevelObjects = true;
+		DetailsViewArgs.bShowKeyablePropertiesOption = false;
+		DetailsViewArgs.bShowModifiedPropertiesOption = false;
+		DetailsViewArgs.bAllowFavoriteSystem = false;
+		DetailsViewArgs.bShowAnimatedPropertiesOption = false;
+	}
+
+	FStructureDetailsViewArgs StructureViewArgs;
+	{
+		StructureViewArgs.bShowObjects = true;
+		StructureViewArgs.bShowAssets = true;
+		StructureViewArgs.bShowClasses = true;
+		StructureViewArgs.bShowInterfaces = true;
+	}
+
+	FeatureDetailsView = PropertyModule.CreateStructureDetailView(DetailsViewArgs, StructureViewArgs, nullptr);
+
 	if (Definition)
 	{
 		Definition->SetFlags(RF_Transactional);
@@ -200,7 +215,7 @@ TSharedRef<SDockTab> FRigMapperDefinitionEditorToolkit::SpawnDetailsTab(const FS
 	.TabColorScale(GetTabColorScale())
 	.Label(LOCTEXT("RigMapperDefinitionEditorToolkitDetailsTab", "Details"))
 	[
-		DetailsView.ToSharedRef()
+		FeatureDetailsView->GetWidget().ToSharedRef()
 	];
 		
 	return DockTab;
@@ -267,9 +282,8 @@ void FRigMapperDefinitionEditorToolkit::HandleFinishedChangingProperties(const F
 	const bool bNullOutput = PropertyName == NullOutputName && ParentName == NullOutputName;
 	const bool bRename = ChangeType == EPropertyChangeType::ValueSet && (bInput || bOutput || bName || bFeaturesChildProperty || bNullOutput);
 	
-	GraphEditor->RebuildGraph();
-	StructureView->RebuildTree();
-	
+	//GraphEditor->RebuildGraph();
+	StructureView->RebuildTree();	
 	if (bSelectNew || bRename)
 	{
 		if (bInput)
@@ -319,22 +333,31 @@ void FRigMapperDefinitionEditorToolkit::HandleFinishedChangingProperties(const F
 
 void FRigMapperDefinitionEditorToolkit::HandleGraphSelectionChanged(const TSet<UObject*>& Nodes)
 {
+	if (!StructureView || !FeatureDetailsView)
+	{
+		return;
+	}
+
 	StructureView->ClearSelection();
-	
+	FeatureDetailsView->SetStructureData(nullptr);
 	for (const UObject* ObjectNode : Nodes)
 	{
 		if (const URigMapperDefinitionEditorGraphNode* Node = Cast<URigMapperDefinitionEditorGraphNode>(ObjectNode))
 		{
-			StructureView->SelectNode(Node->GetNodeName(), Node->GetNodeType(), true);
+			StructureView->SelectNode(Node->GetNodeName(), Node->GetNodeType(), true /*Is selected*/);
+			if (Nodes.Num() == 1)
+			{
+				SetFeatureDetailsView(Node);
+			}					
 		}
 	}
 	
-	DetailsView->ForceRefresh();
+	//DetailsView->ForceRefresh();
 }
 
 void FRigMapperDefinitionEditorToolkit::HandleStructureSelectionChanged(ESelectInfo::Type SelectInfo, TArray<FString> SelectedInputs, TArray<FString> SelectedFeatures, TArray<FString> SelectedOutputs, TArray<FString> SelectedNullOutputs)
 {
-	DetailsView->ForceRefresh();
+	//DetailsView->ForceRefresh();
 
 	if (SelectInfo != ESelectInfo::Direct)
 	{
@@ -342,4 +365,49 @@ void FRigMapperDefinitionEditorToolkit::HandleStructureSelectionChanged(ESelectI
 	};
 }
 
-#undef LOCTEXT_NAMESPACE
+void FRigMapperDefinitionEditorToolkit::SetFeatureDetailsView(const URigMapperDefinitionEditorGraphNode* InNode)
+{
+	URigMapperDefinitionEditorGraphNode* Node = ConstCast<const URigMapperDefinitionEditorGraphNode>(InNode);
+	if (!Node || !Node->GetNodeFeature())
+	{
+		return;
+	}
+
+	if (Node->GetNodeType() == ERigMapperNodeType::SDK)
+	{
+		const TSharedPtr<FStructOnScope> SdkStructure = MakeShared<FStructOnScope>(FRigMapperSdkFeature::StaticStruct(), reinterpret_cast<uint8*>(Node->GetNodeFeature()));
+		FeatureDetailsView->SetStructureData(SdkStructure);
+	}
+	else if (Node->GetNodeType() == ERigMapperNodeType::WeightedSum)
+	{
+		const TSharedPtr<FStructOnScope> WsStructure = MakeShared<FStructOnScope>(FRigMapperWsFeature::StaticStruct(), reinterpret_cast<uint8*>(Node->GetNodeFeature()));
+		FeatureDetailsView->SetStructureData(WsStructure);
+	}
+	else if (Node->GetNodeType() == ERigMapperNodeType::Multiply)
+	{
+		const TSharedPtr<FStructOnScope> MultiplyStructure = MakeShared<FStructOnScope>(FRigMapperMultiplyFeature::StaticStruct(), reinterpret_cast<uint8*>(Node->GetNodeFeature()));
+		FeatureDetailsView->SetStructureData(MultiplyStructure);
+	}
+	else if (Node->GetNodeType() == ERigMapperNodeType::Output)
+	{
+		const TSharedPtr<FStructOnScope> OutputStructure = MakeShared<FStructOnScope>(FRigMapperOutputFeature::StaticStruct(), reinterpret_cast<uint8*>(Node->GetNodeFeature()));
+		FeatureDetailsView->SetStructureData(OutputStructure);
+	}
+	else if (Node->GetNodeType() == ERigMapperNodeType::Input || Node->GetNodeType() == ERigMapperNodeType::NullOutput)
+	{
+		const TSharedPtr<FStructOnScope> FeatureStructure = MakeShared<FStructOnScope>(FRigMapperFeature::StaticStruct(), reinterpret_cast<uint8*>(Node->GetNodeFeature()));
+		FeatureDetailsView->SetStructureData(FeatureStructure);
+	}	
+	FeatureDetailsView->GetOnFinishedChangingPropertiesDelegate().Clear();
+	FeatureDetailsView->GetOnFinishedChangingPropertiesDelegate().AddLambda(
+		[WeakNode = Node]
+		(const FPropertyChangedEvent&)
+		{
+			if (WeakNode)
+			{
+				WeakNode->UpdateNodeFromFeature();
+			}
+		});
+}
+
+#undef LOCTEXT_NAMESPACE
\ No newline at end of file
diff --git a/Engine/Plugins/Experimental/Animation/RigMapper/Source/RigMapperEditor/Private/RigMapperDefinitionEditorToolkit.h b/Engine/Plugins/Experimental/Animation/RigMapper/Source/RigMapperEditor/Private/RigMapperDefinitionEditorToolkit.h
--- a/Engine/Plugins/Experimental/Animation/RigMapper/Source/RigMapperEditor/Private/RigMapperDefinitionEditorToolkit.h
+++ b/Engine/Plugins/Experimental/Animation/RigMapper/Source/RigMapperEditor/Private/RigMapperDefinitionEditorToolkit.h
@@ -11,8 +11,10 @@
 #define UE_API RIGMAPPEREDITOR_API
 
 enum class ERigMapperNodeType : uint8;
+class URigMapperDefinitionEditorGraphNode;
 class SRigMapperDefinitionStructureView;
 class SRigMapperDefinitionGraphEditor;
+class IStructureDetailsView;
 
 /**
  * The toolkit for the URigMapperDefinition asset editor
@@ -56,6 +58,7 @@ private:
 	UE_API void HandleGraphSelectionChanged(const TSet<UObject*>& Nodes);
 	
 	UE_API void HandleStructureSelectionChanged(ESelectInfo::Type SelectInfo, TArray<FString> SelectedInputs, TArray<FString> SelectedFeatures, TArray<FString> SelectedOutputs, TArray<FString> SelectedNullOutputs);
+	void SetFeatureDetailsView(const URigMapperDefinitionEditorGraphNode* InNode);
 
 private:
 	static UE_API const FName DefinitionEditorGraphTabId;
@@ -68,6 +71,7 @@ private:
 
 	TSharedPtr<SRigMapperDefinitionGraphEditor> GraphEditor;
 	TSharedPtr<IDetailsView> DetailsView;
+	TSharedPtr<IStructureDetailsView> FeatureDetailsView;
 	TSharedPtr<SRigMapperDefinitionStructureView> StructureView;
 };
 
diff --git a/Engine/Plugins/Experimental/Animation/RigMapper/Source/RigMapperEditor/Private/RigMapperGraph/RigMapperDefinitionEditorGraph.cpp b/Engine/Plugins/Experimental/Animation/RigMapper/Source/RigMapperEditor/Private/RigMapperGraph/RigMapperDefinitionEditorGraph.cpp
--- a/Engine/Plugins/Experimental/Animation/RigMapper/Source/RigMapperEditor/Private/RigMapperGraph/RigMapperDefinitionEditorGraph.cpp
+++ b/Engine/Plugins/Experimental/Animation/RigMapper/Source/RigMapperEditor/Private/RigMapperGraph/RigMapperDefinitionEditorGraph.cpp
@@ -23,7 +23,7 @@ void URigMapperDefinitionEditorGraph::RebuildGraph()
 	RemoveAllNodes();
 	ConstructNodes();
 	NotifyGraphChanged();
-	RequestRefreshLayout(true);
+	RequestRefreshLayout(false);
 }
 
 void URigMapperDefinitionEditorGraph::ConstructNodes()
@@ -101,7 +101,7 @@ URigMapperDefinitionEditorGraphNode* URigMapperDefinitionEditorGraph::CreateGrap
 		}
 		if (ExistingNode)
 		{
-			return *ExistingNode;		
+			return *ExistingNode;
 		}
 	}
 
@@ -214,6 +214,111 @@ URigMapperDefinitionEditorGraphNode* URigMapperDefinitionEditorGraph::CreateGrap
 	return Node;
 }
 
+
+URigMapperDefinitionEditorGraphNode* URigMapperDefinitionEditorGraph::CreateNewGraphNode(const ERigMapperNodeType InNodeType, UEdGraphPin* InFromPin, const FVector2f& InLocation, bool bSelectNewNode /*= true*/)
+{
+	if (!WeakDefinition.IsValid())
+	{
+		return nullptr;
+	}
+	FString FromPinName = "New node";
+	if (InFromPin)
+	{
+		FromPinName = Cast<URigMapperDefinitionEditorGraphNode>(InFromPin->GetOwningNode())->GetNodeName();
+	}
+	URigMapperDefinitionEditorGraphNode* NewNode = CreateGraphNode(MakeUniqueObjectName(this, URigMapperDefinitionEditorGraphNode::StaticClass(), FName(FromPinName)).ToString(), InNodeType);
+	NewNode->NodePosX = InLocation.X;
+	NewNode->NodePosY = InLocation.Y;
+	FString NodeName = NewNode->GetNodeName();
+	if (InNodeType != ERigMapperNodeType::Input && InNodeType != ERigMapperNodeType::NullOutput && InNodeType != ERigMapperNodeType::Invalid)
+	{
+		UEdGraphPin* InPin = NewNode->CreateInputPin();
+		InPin->bHidden = false;
+	}
+	FString InputNodeName = "";
+	if (InFromPin && InFromPin->Direction == EGPD_Output)
+	{
+		InputNodeName = FromPinName;
+	}
+	if (InNodeType == ERigMapperNodeType::Multiply)
+	{
+		FRigMapperMultiplyFeature NewFeature = FRigMapperMultiplyFeature(NodeName);
+		if (!InputNodeName.IsEmpty())
+		{
+			NewFeature.Inputs.Add(InputNodeName);
+		}
+		WeakDefinition->Features.Multiply.Add(NewFeature);		
+		// Multiply and Weighted Sum features requires at least 2 inputs.
+		UEdGraphPin* InPin = NewNode->CreateInputPin();
+		InPin->bHidden = false;
+	}
+	else if (InNodeType == ERigMapperNodeType::WeightedSum)
+	{
+		FRigMapperWsFeature NewFeature = FRigMapperWsFeature(NodeName);
+		if (!InputNodeName.IsEmpty())
+		{
+			NewFeature.Inputs.Add(InputNodeName);
+		}
+		WeakDefinition->Features.WeightedSums.Add(NewFeature);
+	}
+	else if (InNodeType == ERigMapperNodeType::SDK)
+	{
+		FRigMapperSdkFeature NewFeature = FRigMapperSdkFeature(NodeName);
+		if (!InputNodeName.IsEmpty())
+		{
+			NewFeature.Input = InputNodeName;
+		}
+		WeakDefinition->Features.SDKs.Add(NewFeature);
+	}
+	else if (InNodeType == ERigMapperNodeType::Output)
+	{
+		WeakDefinition->Outputs.Add(NodeName, InputNodeName);
+	}
+	else if (InNodeType == ERigMapperNodeType::Input)
+	{
+		WeakDefinition->Inputs.Add(NodeName);
+	}
+	else if (InNodeType == ERigMapperNodeType::NullOutput)
+	{
+		WeakDefinition->NullOutputs.Add(NodeName);
+	}
+	if (InNodeType != ERigMapperNodeType::Output && InNodeType != ERigMapperNodeType::NullOutput)
+	{
+		UEdGraphPin* OutPin = NewNode->CreateOutputPin();
+		OutPin->bHidden = false;
+	}
+
+	NewNode->SetupNode(WeakDefinition.Get(), NodeName, InNodeType);
+
+	if (InFromPin)
+	{
+		const UEdGraphSchema* GraphSchema = GetSchema();
+		UEdGraphPin* TargetPin = nullptr;
+		if (InFromPin->Direction == EGPD_Output)
+		{
+			const TArray<UEdGraphPin*>& Pins = NewNode->GetInputPins();
+			if (Pins.Num() > 0)
+			{
+				GraphSchema->TryCreateConnection(InFromPin, Pins[0]);
+			}
+		}
+		else
+		{
+			UEdGraphPin* OutPin = NewNode->GetOutputPin();
+			if (OutPin)
+			{
+				GraphSchema->TryCreateConnection(OutPin, InFromPin);
+				if (OutPin->LinkedTo.Num() > 0)
+				{
+					Cast<URigMapperDefinitionEditorGraphNode>(OutPin->LinkedTo[0]->GetOwningNode())->AddNewInputLink(NodeName);
+				}
+			}
+		}
+	}
+
+	return NewNode;
+}
+
 void URigMapperDefinitionEditorGraph::LayoutNodeRec(URigMapperDefinitionEditorGraphNode* InNode, double InputsWidth, double PosY, TArray<URigMapperDefinitionEditorGraphNode*>& LayedOutNodes) const
 {
 	const int32 NodeMarginX = 20;
@@ -223,7 +328,7 @@ void URigMapperDefinitionEditorGraph::LayoutNodeRec(URigMapperDefinitionEditorGr
 	double SubGraphHeight = 0;
 	
 	for (UEdGraphPin* InPin : InNode->GetInputPins())
-	{
+	{ 
 		for (const UEdGraphPin* OutPin : InPin->LinkedTo)
 		{
 			LinkedNode = Cast<URigMapperDefinitionEditorGraphNode>(OutPin->GetOwningNode());
diff --git a/Engine/Plugins/Experimental/Animation/RigMapper/Source/RigMapperEditor/Private/RigMapperGraph/RigMapperDefinitionEditorGraph.h b/Engine/Plugins/Experimental/Animation/RigMapper/Source/RigMapperEditor/Private/RigMapperGraph/RigMapperDefinitionEditorGraph.h
--- a/Engine/Plugins/Experimental/Animation/RigMapper/Source/RigMapperEditor/Private/RigMapperGraph/RigMapperDefinitionEditorGraph.h
+++ b/Engine/Plugins/Experimental/Animation/RigMapper/Source/RigMapperEditor/Private/RigMapperGraph/RigMapperDefinitionEditorGraph.h
@@ -32,6 +32,8 @@ public:
 	
 	URigMapperDefinition* GetDefinition() const { return WeakDefinition.Get(); };
 
+	UE_API URigMapperDefinitionEditorGraphNode* CreateGraphNode(const FString& NodeName, ERigMapperNodeType NodeType);
+	UE_API URigMapperDefinitionEditorGraphNode* CreateNewGraphNode(const ERigMapperNodeType InNodeType, UEdGraphPin* InFromPin, const FVector2f& InLocation, bool bSelectNewNode = true);
 	UE_API TArray<URigMapperDefinitionEditorGraphNode*> GetNodesByName(const TArray<FString>& Inputs, const TArray<FString>& Features, const TArray<FString>& Outputs, const TArray<FString>& NullOutputs) const;
 
 private:
@@ -39,7 +41,6 @@ private:
 	UE_API URigMapperDefinitionEditorGraphNode* CreateOutputNode(URigMapperDefinition* Definition, const FString& NodeName);
 	UE_API URigMapperDefinitionEditorGraphNode* CreateFeatureNode(URigMapperDefinition* Definition, const FString& NodeName);
 	static UE_API void LinkGraphNodes(URigMapperDefinitionEditorGraphNode* InNode, URigMapperDefinitionEditorGraphNode* OutNode);
-	UE_API URigMapperDefinitionEditorGraphNode* CreateGraphNode(const FString& NodeName, ERigMapperNodeType NodeType);
 	UE_API void LayoutNodeRec(URigMapperDefinitionEditorGraphNode* InNode, double InputsWidth, double PosY, TArray<URigMapperDefinitionEditorGraphNode*>& LayedOutNodes) const;
 	
 private:
diff --git a/Engine/Plugins/Experimental/Animation/RigMapper/Source/RigMapperEditor/Private/RigMapperGraph/RigMapperDefinitionEditorGraphNode.cpp b/Engine/Plugins/Experimental/Animation/RigMapper/Source/RigMapperEditor/Private/RigMapperGraph/RigMapperDefinitionEditorGraphNode.cpp
--- a/Engine/Plugins/Experimental/Animation/RigMapper/Source/RigMapperEditor/Private/RigMapperGraph/RigMapperDefinitionEditorGraphNode.cpp
+++ b/Engine/Plugins/Experimental/Animation/RigMapper/Source/RigMapperEditor/Private/RigMapperGraph/RigMapperDefinitionEditorGraphNode.cpp
@@ -18,7 +18,7 @@ TSharedPtr<SGraphNode> URigMapperDefinitionEditorGraphNode::NodeFactory::CreateN
 	{
 		if (URigMapperDefinitionEditorGraph* Graph = Cast<URigMapperDefinitionEditorGraph>(Node->GetGraph()))
 		{
-			Graph->RequestRefreshLayout(true);
+			Graph->RequestRefreshLayout(false);
 		}
 
 		TSharedRef<SGraphNode> GraphNode = SNew(SRigMapperDefinitionGraphEditorNode, ThisNode);
@@ -63,79 +63,250 @@ FLinearColor URigMapperDefinitionEditorGraphNode::GetNodeTitleColor() const
 	return FLinearColor(0.969, 0.5, 0.466);
 }
 
+void URigMapperDefinitionEditorGraphNode::GetNodeContextMenuActions(class UToolMenu* Menu, class UGraphNodeContextMenuContext* Context) const
+{
+
+}
+
+void URigMapperDefinitionEditorGraphNode::AllocateDefaultPins()
+{
+
+}
+
 void URigMapperDefinitionEditorGraphNode::SetupNode(URigMapperDefinition* InDefinition, const FString& InNodeName, ERigMapperNodeType InNodeType)
 {
 	NodeName = InNodeName;
 	Definition = InDefinition;
-	
+	NodeType = InNodeType;
 	NodeTitle = FText::FromString(NodeName);
-	SetNodeType(InNodeType);
-}
 
-void URigMapperDefinitionEditorGraphNode::SetNodeType(ERigMapperNodeType InNodeType)
-{
-	NodeType = InNodeType;
+	if (!Definition.IsValid())
+	{
+		return;
+	}
+	FeaturePtr = nullptr;
+	IsExternalOwner = false;
 
-	if (Definition.IsValid() && (NodeType == ERigMapperNodeType::WeightedSum || NodeType == ERigMapperNodeType::SDK))
+	if (NodeType == ERigMapperNodeType::Input && Definition.Get()->Inputs.Contains(NodeName))
+	{
+		FeaturePtr = new FRigMapperFeature(NodeName);		
+	}
+	else if (NodeType == ERigMapperNodeType::WeightedSum || NodeType == ERigMapperNodeType::SDK || NodeType == ERigMapperNodeType::Multiply)
 	{
+		IsExternalOwner = true;
 		ERigMapperFeatureType FeatureType;
 		if (FRigMapperFeature* Feature = Definition.Get()->Features.Find(NodeName, FeatureType))
 		{
-			if (FeatureType == ERigMapperFeatureType::SDK)
+			FeaturePtr = Feature;
+		}
+	}
+	else if (NodeType == ERigMapperNodeType::Output)
+	{
+		if (FString* LinkedInputName = Definition.Get()->Outputs.Find(NodeName))
+		{
+			FRigMapperOutputFeature* OutputFeature = new FRigMapperOutputFeature(NodeName);
+			OutputFeature->Input = *LinkedInputName;			
+			FeaturePtr = static_cast<FRigMapperFeature*>(OutputFeature);
+		}
+	}
+	else if (NodeType == ERigMapperNodeType::NullOutput && Definition.Get()->NullOutputs.Contains(NodeName))
+	{
+		FeaturePtr = new FRigMapperNullOutputFeature(NodeName);
+	}
+
+	UpdateNodeFields();
+}
+
+void URigMapperDefinitionEditorGraphNode::UpdateNodeFields()
+{
+
+	if (!Definition.IsValid())
+	{
+		return;
+	}
+
+	if (!FeaturePtr)
+	{
+		return;
+	}
+
+	if (NodeType == ERigMapperNodeType::WeightedSum || NodeType == ERigMapperNodeType::SDK)
+	{
+		if (FeaturePtr->GetFeatureType() == ERigMapperFeatureType::SDK)
+		{
+			TArray<FString> InKeys;
+			TArray<FString> OutKeys;
+
+			const TArray<FRigMapperSdkKey>& Keys = static_cast<FRigMapperSdkFeature*>(FeaturePtr)->Keys;
+
+			InKeys.Reserve(Keys.Num());
+			OutKeys.Reserve(Keys.Num());
+			for (const FRigMapperSdkKey& Key : Keys)
 			{
-				TArray<FString> InKeys;
-				TArray<FString> OutKeys;
-				
-				const TArray<FRigMapperSdkKey>& Keys = static_cast<FRigMapperSdkFeature*>(Feature)->Keys;
-
-				InKeys.Reserve(Keys.Num());
-				OutKeys.Reserve(Keys.Num());
-				for (const FRigMapperSdkKey& Key : Keys)
-				{
-					InKeys.Add(FString::Printf(TEXT("%.*f"), 3, Key.In));
-					OutKeys.Add(FString::Printf(TEXT("%.*f"), 3, Key.Out));
-				}
-				const FString In = FString::Join(InKeys, TEXT(", "));
-				const FString Out = FString::Join(OutKeys, TEXT(", "));
-				NodeSubtitle = FText::FromString(FString::Printf(TEXT("[%s] > [%s]"), *In, *Out));
+				InKeys.Add(FString::Printf(TEXT("%.*f"), 3, Key.In));
+				OutKeys.Add(FString::Printf(TEXT("%.*f"), 3, Key.Out));
 			}
-			if (FeatureType == ERigMapperFeatureType::WeightedSum)
-			{
-				const TMap<FString, double>& Inputs = static_cast<FRigMapperWsFeature*>(Feature)->Inputs;
-				const FRigMapperFeatureRange& Range = static_cast<FRigMapperWsFeature*>(Feature)->Range;
-				const FText MinText = LOCTEXT("RangeMinimum", "min");
-				const FText MaxText = LOCTEXT("RangeMaximum", "max");
+			const FString In = FString::Join(InKeys, TEXT(", "));
+			const FString Out = FString::Join(OutKeys, TEXT(", "));
+			NodeSubtitle = FText::FromString(FString::Printf(TEXT("[%s] > [%s]"), *In, *Out));
+		}
+		if (FeaturePtr->GetFeatureType() == ERigMapperFeatureType::WeightedSum)
+		{
+			const FRigMapperWsFeature* Feature = static_cast<FRigMapperWsFeature*>(FeaturePtr);
+			const TMap<FString, double>& Inputs = Feature->Inputs;
+			const FRigMapperFeatureRange& Range = Feature->Range;
+			const FText MinText = LOCTEXT("RangeMinimum", "min");
+			const FText MaxText = LOCTEXT("RangeMaximum", "max");
 
-				for (UEdGraphPin* InPin : InputPins)
+			for (UEdGraphPin* InPin : InputPins)
+			{
+				for (const UEdGraphPin* OutPin : InPin->LinkedTo)
 				{
-					for (const UEdGraphPin* OutPin : InPin->LinkedTo)
+					const FString LinkedNodeName = OutPin->GetOwningNode()->GetNodeTitle(ENodeTitleType::FullTitle).ToString();
+					if (Inputs.Contains(LinkedNodeName))
 					{
-						const FString LinkedNodeName = OutPin->GetOwningNode()->GetNodeTitle(ENodeTitleType::FullTitle).ToString();
-						if (Inputs.Contains(LinkedNodeName))
+						InPin->PinFriendlyName = FText::Format(INVTEXT("{0}"), Inputs[LinkedNodeName]);
+						FString RangeStr;
+						if (Range.bHasLowerBound)
 						{
-							InPin->PinFriendlyName = FText::Format(INVTEXT("{0}"), Inputs[LinkedNodeName]);
-							FString RangeStr;
-							if (Range.bHasLowerBound)
-							{
-								RangeStr.Appendf(TEXT("%s: %.*f"), *MinText.ToString(), 3, Range.LowerBound);
-							}
-							if (Range.bHasLowerBound && Range.bHasUpperBound)
-							{
-								RangeStr.Append(TEXT("\n"));
-							}
-							if (Range.bHasUpperBound)
-							{
-								RangeStr.Appendf(TEXT("%s: %.*f"), *MaxText.ToString(), 3, Range.UpperBound);
-							}
-							NodeSubtitle = FText::FromString(RangeStr);
-						}	
+							RangeStr.Appendf(TEXT("%s: %.*f"), *MinText.ToString(), 3, Range.LowerBound);
+						}
+						if (Range.bHasLowerBound && Range.bHasUpperBound)
+						{
+							RangeStr.Append(TEXT("\n"));
+						}
+						if (Range.bHasUpperBound)
+						{
+							RangeStr.Appendf(TEXT("%s: %.*f"), *MaxText.ToString(), 3, Range.UpperBound);
+						}
+						NodeSubtitle = FText::FromString(RangeStr);
 					}
 				}
 			}
+
 		}
-		
+	}
+}
+
+void URigMapperDefinitionEditorGraphNode::RelinkFeature(const FString& InOldInput, const FString& InNewInput)
+{
+	if (!Definition.IsValid() || !FeaturePtr)
+	{
+		return;
 	}
 
+	if (NodeType == ERigMapperNodeType::Output)
+	{
+		FRigMapperOutputFeature* Feature = static_cast<FRigMapperOutputFeature*>(FeaturePtr);
+		if (Feature)
+		{
+			Feature->Input = InNewInput;
+		}
+		if (Definition->Outputs.Contains(NodeName))
+		{
+			Definition->Outputs[NodeName] = InNewInput;
+		}
+	}
+	else if (NodeType == ERigMapperNodeType::Multiply)
+	{
+		FRigMapperMultiplyFeature* Feature = static_cast<FRigMapperMultiplyFeature*>(FeaturePtr);
+		if (Feature)
+		{
+			for (int32 i = 0; i < Feature->Inputs.Num(); i++)
+			{
+				if (Feature->Inputs[i] == InOldInput)
+				{
+					Feature->Inputs[i] = InNewInput;
+				}
+			}
+		}
+	}
+	else if (NodeType == ERigMapperNodeType::WeightedSum)
+	{
+		FRigMapperWsFeature* Feature = static_cast<FRigMapperWsFeature*>(FeaturePtr);
+		if (Feature)
+		{
+			if (double* Weight = Feature->Inputs.Find(InOldInput))
+			{
+				Feature->Inputs.Remove(InOldInput);
+				Feature->Inputs.Add(InNewInput, MoveTemp(*Weight));
+			}
+		}
+	}
+	else if (NodeType == ERigMapperNodeType::SDK)
+	{
+		FRigMapperSdkFeature* Feature = static_cast<FRigMapperSdkFeature*>(FeaturePtr);
+		if (Feature)
+		{
+			Feature->Input = InNewInput;
+		}
+	}	
+}
+
+void URigMapperDefinitionEditorGraphNode::AddNewInputLink(const FString& InNewInput)
+{
+	if (!Definition.IsValid() || !FeaturePtr)
+	{
+		return;
+	}
+
+	if (NodeType == ERigMapperNodeType::Output)
+	{
+		FRigMapperOutputFeature* Feature = static_cast<FRigMapperOutputFeature*>(FeaturePtr);
+		if (Feature)
+		{
+			Feature->Input = InNewInput;
+		}
+		if (Definition->Outputs.Contains(NodeName))
+		{
+			Definition->Outputs[NodeName] = InNewInput;
+		}
+	}
+	else if (NodeType == ERigMapperNodeType::Multiply)
+	{
+		FRigMapperMultiplyFeature* Feature = static_cast<FRigMapperMultiplyFeature*>(FeaturePtr);
+		if (Feature)
+		{
+			Feature->Inputs.Add(InNewInput);
+		}
+	}
+	else if (NodeType == ERigMapperNodeType::WeightedSum)
+	{
+		FRigMapperWsFeature* Feature = static_cast<FRigMapperWsFeature*>(FeaturePtr);
+		if (Feature)
+		{
+			if (!Feature->Inputs.Contains(InNewInput))
+			{				
+				Feature->Inputs.Add(InNewInput, 1.0f);
+			}
+		}
+	}
+	else if (NodeType == ERigMapperNodeType::SDK)
+	{
+		FRigMapperSdkFeature* Feature = static_cast<FRigMapperSdkFeature*>(FeaturePtr);
+		if (Feature)
+		{
+			Feature->Input = InNewInput;
+		}
+	}
+}
+
+void URigMapperDefinitionEditorGraphNode::BeginDestroy()
+{
+	if (FeaturePtr)
+	{
+		if (!IsExternalOwner)
+		{
+			delete(FeaturePtr);
+		}
+		FeaturePtr = nullptr;
+	}
+	Super::BeginDestroy();
+}
+
+FRigMapperFeature* URigMapperDefinitionEditorGraphNode::GetNodeFeature() const
+{
+	return FeaturePtr;
 }
 
 UEdGraphPin* URigMapperDefinitionEditorGraphNode::CreateInputPin()
@@ -156,4 +327,31 @@ void URigMapperDefinitionEditorGraphNode::GetRect(FVector2D& TopLeft, FVector2D&
 	BottomRight.Y = NodePosY + GetDimensions().Y + GetMargin().Y;
 }
 
+void URigMapperDefinitionEditorGraphNode::UpdateNodeFromFeature()
+{
+	if (FeaturePtr)
+	{
+		FString NewName = FeaturePtr->Name;
+		
+		if (NewName != NodeName && NodeType != ERigMapperNodeType::Output && NodeType != ERigMapperNodeType::NullOutput)
+		{
+			// Update inputs of all nodes directly connected to this node.
+			for (UEdGraphPin* OutPin : OutputPins) // It will always be one output pin.
+			{
+				for (UEdGraphPin* LinkedPin : OutPin->LinkedTo)
+				{
+					if (URigMapperDefinitionEditorGraphNode* LinkedNode = Cast<URigMapperDefinitionEditorGraphNode>(LinkedPin->GetOwningNode()))
+					{
+						LinkedNode->RelinkFeature(NodeName, NewName);
+					}
+				}
+			}
+		}
+
+		NodeName = NewName;
+		NodeTitle = FText::FromString(NodeName);
+		UpdateNodeFields();
+	}
+}
+
 #undef LOCTEXT_NAMESPACE
diff --git a/Engine/Plugins/Experimental/Animation/RigMapper/Source/RigMapperEditor/Private/RigMapperGraph/RigMapperDefinitionEditorGraphNode.h b/Engine/Plugins/Experimental/Animation/RigMapper/Source/RigMapperEditor/Private/RigMapperGraph/RigMapperDefinitionEditorGraphNode.h
--- a/Engine/Plugins/Experimental/Animation/RigMapper/Source/RigMapperEditor/Private/RigMapperGraph/RigMapperDefinitionEditorGraphNode.h
+++ b/Engine/Plugins/Experimental/Animation/RigMapper/Source/RigMapperEditor/Private/RigMapperGraph/RigMapperDefinitionEditorGraphNode.h
@@ -11,6 +11,7 @@
 #define UE_API RIGMAPPEREDITOR_API
 
 class URigMapperDefinition;
+struct FRigMapperFeature;
 
 UENUM()
 enum class ERigMapperNodeType : uint8
@@ -43,13 +44,17 @@ public:
 	virtual FText GetNodeTitle(ENodeTitleType::Type TitleType) const override { return NodeTitle; }
 	UE_API virtual FLinearColor GetNodeBodyTintColor() const override;
 	UE_API virtual FLinearColor GetNodeTitleColor() const override;
+	UE_API virtual void GetNodeContextMenuActions(class UToolMenu* Menu, class UGraphNodeContextMenuContext* Context) const override;
+	UE_API virtual void AllocateDefaultPins() override;
+	UE_API void BeginDestroy() override;
 	// !UEdGraphNode implementation
 
 	UE_API void SetupNode(URigMapperDefinition* InDefinition, const FString& InNodeName, ERigMapperNodeType InNodeType);
 	const FString& GetNodeName() const { return NodeName; }
 	
-	UE_API void SetNodeType(ERigMapperNodeType InNodeType);
 	ERigMapperNodeType GetNodeType() const { return NodeType; }
+
+	UE_API FRigMapperFeature* GetNodeFeature() const;
 	
 	const FText& GetSubtitle() const { return NodeSubtitle; }
 	
@@ -66,6 +71,15 @@ public:
 
 	UE_API void GetRect(FVector2D& TopLeft, FVector2D& BottomRight) const;
 
+	/** Updates node from definition and makes it dirty */
+	UE_API void UpdateNodeFromFeature();
+
+	/** Updates node body */
+	UE_API void UpdateNodeFields();
+
+	UE_API void RelinkFeature(const FString& InOldInput, const FString& InNewInput);
+	UE_API void AddNewInputLink(const FString& InNewInput);
+
 protected:
 	/** Cached title for the node */
 	FText NodeTitle;
@@ -73,10 +87,10 @@ protected:
 	/** Cached title for the node */
 	FText NodeSubtitle;
 
-	/** Our one input pin */
+	/** Our input pins */
 	TArray<UEdGraphPin*> InputPins;
 
-	/** Our one output pin */
+	/** Our output pins */
 	TArray<UEdGraphPin*> OutputPins;
 
 	/** Cached dimensions of this node (used for layout) */
@@ -87,9 +101,13 @@ protected:
 	
 	TWeakObjectPtr<URigMapperDefinition> Definition;
 
+	FRigMapperFeature* FeaturePtr = nullptr;
+
 	FString NodeName;
 
 	ERigMapperNodeType NodeType = ERigMapperNodeType::Invalid;
+
+	bool IsExternalOwner;
 };
 
 #undef UE_API
diff --git a/Engine/Plugins/Experimental/Animation/RigMapper/Source/RigMapperEditor/Private/RigMapperGraph/RigMapperDefinitionEditorGraphSchema.cpp b/Engine/Plugins/Experimental/Animation/RigMapper/Source/RigMapperEditor/Private/RigMapperGraph/RigMapperDefinitionEditorGraphSchema.cpp
--- a/Engine/Plugins/Experimental/Animation/RigMapper/Source/RigMapperEditor/Private/RigMapperGraph/RigMapperDefinitionEditorGraphSchema.cpp
+++ b/Engine/Plugins/Experimental/Animation/RigMapper/Source/RigMapperEditor/Private/RigMapperGraph/RigMapperDefinitionEditorGraphSchema.cpp
@@ -2,43 +2,65 @@
 
 
 #include "RigMapperDefinitionEditorGraphSchema.h"
-
+#include "Framework/Commands/GenericCommands.h"
+#include "GraphEditorActions.h"
 #include "RigMapperDefinition.h"
 #include "RigMapperDefinitionEditorGraphNode.h"
+#include "RigMapperDefinitionEditorGraph.h"
 #include "SRigMapperDefinitionGraphEditorNode.h"
-
+#include "ToolMenu.h"
+#include "ScopedTransaction.h"
 #include "EdGraph/EdGraph.h"
 
 #include UE_INLINE_GENERATED_CPP_BY_NAME(RigMapperDefinitionEditorGraphSchema)
 
-
 #define LOCTEXT_NAMESPACE "RigMapperDefinitionEditorGraphSchema"
 
+FSchemaAction_SpawnRigMapperNode::FSchemaAction_SpawnRigMapperNode() :
+	FEdGraphSchemaAction(),
+	NodeType(ERigMapperNodeType::NullOutput)
+{
+}
+
+FSchemaAction_SpawnRigMapperNode::FSchemaAction_SpawnRigMapperNode(ERigMapperNodeType InNodeType, const FText& InCategory, const FText& InDescription, const FText& InTooltip, const int32 InGrouping) :
+	FEdGraphSchemaAction(InCategory, InDescription, InTooltip, InGrouping),
+	NodeType(InNodeType)
+{
+}
+
+UEdGraphNode* FSchemaAction_SpawnRigMapperNode::PerformAction(UEdGraph* ParentGraph, UEdGraphPin* FromPin, const FVector2f& Location, bool bSelectNewNode /*= true*/)
+{
+	URigMapperDefinitionEditorGraphNode* NewNode = nullptr;
+	if (URigMapperDefinitionEditorGraph* Graph = Cast<URigMapperDefinitionEditorGraph>(ParentGraph))
+	{
+		TSharedPtr<FScopedTransaction> Transaction = MakeShared<FScopedTransaction>(LOCTEXT("RigMapperDefinitionEditor_SpawnNewNode", "Spawn RigMapper definition node"));
+		Graph->Modify();
+
+		NewNode = Graph->CreateNewGraphNode(NodeType, FromPin, Location, bSelectNewNode);
+		
+		Transaction.Reset();
+	}
+	return NewNode;
+}
 
 void URigMapperDefinitionEditorGraphSchema::GetContextMenuActions(UToolMenu* Menu, UGraphNodeContextMenuContext* Context) const
 {
-	// Menu->AddDynamicSection("PhysicsAssetGraphSchema", FNewToolMenuDelegateLegacy::CreateLambda([](FMenuBuilder& MenuBuilder, UToolMenu* InMenu)
-	// {
-	// 	UGraphNodeContextMenuContext* ContextObject = InMenu->FindContext<UGraphNodeContextMenuContext>();
-	// 	if (!ContextObject)
-	// 	{
-	// 		return;
-	// 	}
-	//
-	// 	const UPhysicsAssetGraph* PhysicsAssetGraph = CastChecked<const UPhysicsAssetGraph>(ContextObject->Graph);
-	// 	TSharedPtr<FPhysicsAssetEditorSharedData> SharedData = PhysicsAssetGraph->GetPhysicsAssetEditor()->GetSharedData();
-	//
-	// 	if (const UPhysicsAssetGraphNode_Constraint* ConstraintNode = Cast<const UPhysicsAssetGraphNode_Constraint>(ContextObject->Node))
-	// 	{
-	// 		PhysicsAssetGraph->GetPhysicsAssetEditor()->BuildMenuWidgetConstraint(MenuBuilder);
-	// 	}
-	// 	else if (const UPhysicsAssetGraphNode_Bone* BoneNode = Cast<const UPhysicsAssetGraphNode_Bone>(ContextObject->Node))
-	// 	{
-	// 		PhysicsAssetGraph->GetPhysicsAssetEditor()->BuildMenuWidgetBody(MenuBuilder);
-	// 	}
-	//
-	// 	PhysicsAssetGraph->GetPhysicsAssetEditor()->BuildMenuWidgetSelection(MenuBuilder);
-	// }));
+	if (!Context->Pin && Context->Node && Context->Node->IsA<URigMapperDefinitionEditorGraphNode>())
+	{
+		// TODO: Implement context menu actions.
+		FToolMenuSection& Section = Menu->AddSection("RigMapperDefinitionEditorNodeActions", LOCTEXT("RigMapperDefinitionEditor_NodeActionsMenuHeader", "Node Actions"));
+		Section.AddMenuEntry(FGenericCommands::Get().Delete);
+		Section.AddMenuEntry(FGenericCommands::Get().Cut);
+		Section.AddMenuEntry(FGenericCommands::Get().Copy);
+		Section.AddMenuEntry(FGenericCommands::Get().Duplicate);
+		Section.AddMenuEntry(FGenericCommands::Get().Rename);
+		Section.AddMenuEntry(FGraphEditorCommands::Get().BreakNodeLinks);
+		Section.AddMenuEntry(FGraphEditorCommands::Get().FindReferences, LOCTEXT("RigMapperDefinitionEditor_FindReferences", "Find References In Graph"),
+			LOCTEXT("RigMapperDefinitionEditor_FindReferencesTooltip", "Find References to the selected Node in the current Graph"),
+			FSlateIcon());
+
+	}
+	Super::GetContextMenuActions(Menu, Context);	
 }
 
 FLinearColor URigMapperDefinitionEditorGraphSchema::GetPinTypeColor(const FEdGraphPinType& PinType) const
@@ -48,18 +70,17 @@ FLinearColor URigMapperDefinitionEditorGraphSchema::GetPinTypeColor(const FEdGra
 
 void URigMapperDefinitionEditorGraphSchema::BreakPinLinks(UEdGraphPin& TargetPin, bool bSendsNodeNotification) const
 {
-	// Don't allow breaking any links
+	Super::BreakPinLinks(TargetPin, bSendsNodeNotification);
 }
 
 void URigMapperDefinitionEditorGraphSchema::BreakSinglePinLink(UEdGraphPin* SourcePin, UEdGraphPin* TargetPin) const
 {
-	// Don't allow breaking any links
+	Super::BreakSinglePinLink(SourcePin, TargetPin);
 }
 
 FPinConnectionResponse URigMapperDefinitionEditorGraphSchema::MovePinLinks(UEdGraphPin& MoveFromPin, UEdGraphPin& MoveToPin, bool bIsIntermediateMove, bool bNotifyLinkedNodes) const
 {
-	// Don't allow moving any links
-	return FPinConnectionResponse(CONNECT_RESPONSE_DISALLOW, TEXT(""));
+	return FPinConnectionResponse(CONNECT_RESPONSE_MAKE, TEXT(""));
 }
 
 FPinConnectionResponse URigMapperDefinitionEditorGraphSchema::CopyPinLinks(UEdGraphPin& CopyFromPin, UEdGraphPin& CopyToPin, bool bIsIntermediateCopy) const
@@ -73,28 +94,160 @@ FConnectionDrawingPolicy* URigMapperDefinitionEditorGraphSchema::CreateConnectio
 	return new ConnectionDrawingPolicy(InBackLayerID, InFrontLayerID, InZoomFactor, InClippingRect, InDrawElements);
 }
 
-const FPinConnectionResponse URigMapperDefinitionEditorGraphSchema::CanCreateConnection(const UEdGraphPin* A, const UEdGraphPin* B) const
+const FPinConnectionResponse URigMapperDefinitionEditorGraphSchema::CanCreateConnection(const UEdGraphPin* InA, const UEdGraphPin* InB) const
 {
-	return FPinConnectionResponse(CONNECT_RESPONSE_DISALLOW, LOCTEXT("RigMapperDefinitionEditor_CannotCreateConnection", "It is currently not possible to create new connections from the graph"));
+	// Make sure the pins are not on the same node
+	if (InA->GetOwningNode() == InB->GetOwningNode())
+	{
+		return FPinConnectionResponse(CONNECT_RESPONSE_DISALLOW, LOCTEXT("RigMapperDefinitionEditor_ConnectionSameNode", "Both are on the same node"));
+	}
+
+	// Compare the directions
+	const UEdGraphPin* InputPin = NULL;
+	const UEdGraphPin* OutputPin = NULL;
+
+	if (!CategorizePinsByDirection(InA, InB, /*out*/ InputPin, /*out*/ OutputPin))
+	{
+		return FPinConnectionResponse(CONNECT_RESPONSE_DISALLOW, LOCTEXT("RigMapperDefinitionEditor_ConnectionIncompatible", "Directions are not compatible"));
+	}
+
+	// Check for new and existing loops
+	FText ResponseMessage;
+	//if (ConnectionCausesLoop(InputPin, OutputPin))
+	//{
+	//	ResponseMessage = LOCTEXT("ConnectionLoop", "Connection could cause loop");
+	//}
+
+	// For non-exec pins, break existing connections on inputs only - multiple output connections are acceptable
+	if (InputPin->LinkedTo.Num() > 0)
+	{	
+		ECanCreateConnectionResponse ReplyBreakOutputs;
+		if (InputPin == InA)
+		{
+			ReplyBreakOutputs = CONNECT_RESPONSE_BREAK_OTHERS_A;
+		}
+		else
+		{
+			ReplyBreakOutputs = CONNECT_RESPONSE_BREAK_OTHERS_B;
+		}
+		if (ResponseMessage.IsEmpty())
+		{
+			ResponseMessage = LOCTEXT("RigMapperDefinitionEditor_ConnectionReplace", "Replace existing connections");
+		}
+		return FPinConnectionResponse(ReplyBreakOutputs, ResponseMessage);
+
+	}
+
+	return FPinConnectionResponse(CONNECT_RESPONSE_MAKE, TEXT(""));
 }
 
 FPinConnectionResponse URigMapperDefinitionEditorGraphSchema::CanCreateNewNodes(UEdGraphPin* InSourcePin) const
 {
-	// if(URigMapperDefinitionEditorGraphNode* Node = Cast<URigMapperDefinitionEditorGraphNode>(InSourcePin->GetOwningNode()))
-	// {
-	// 	if(InSourcePin->Direction == EGPD_Output)
-	// 	{
-	// 		return FPinConnectionResponse(CONNECT_RESPONSE_MAKE, LOCTEXT("MakeANewConstraint", "Create a new constraint"));
-	// 	}
-	// }
-
-	return FPinConnectionResponse(CONNECT_RESPONSE_DISALLOW, LOCTEXT("RigMapperDefinitionEditor_CannotCreateNode", "It is currently not possible to create new nodes from the graph"));
+	return FPinConnectionResponse(CONNECT_RESPONSE_MAKE, TEXT(""));
 }
 
 bool URigMapperDefinitionEditorGraphSchema::SupportsDropPinOnNode(UEdGraphNode* InTargetNode, const FEdGraphPinType& InSourcePinType, EEdGraphPinDirection InSourcePinDirection, FText& OutErrorMessage) const
 {
-	OutErrorMessage = LOCTEXT("RigMapperDefinitionEditor_CannotCreateConnection_DropPinOnNode", "It is currently not possible to create new connections or nodes from the graph");
+	if (InSourcePinDirection == EGPD_Input)
+	{
+		OutErrorMessage = LOCTEXT("RigMapperDefinitionEditor_CannotCreateOutputConnection_DropPinOnNode", "It is not possible to create additional output connections on nodes.");
+		return false;
+	}
+
+	if (URigMapperDefinitionEditorGraphNode* GraphNode = Cast<URigMapperDefinitionEditorGraphNode>(InTargetNode))
+	{
+		if (GraphNode->GetNodeType() == ERigMapperNodeType::Multiply || GraphNode->GetNodeType() == ERigMapperNodeType::WeightedSum)
+		{
+			OutErrorMessage = LOCTEXT("RigMapperDefinitionEditor_CreateNewConnection_DropPinOnNode", "Create new input pin and connect.");
+			return true;
+		}
+		OutErrorMessage = LOCTEXT("RigMapperDefinitionEditor_CannotCreateConnection_DropPinOnNode", "It is not possible to create new connections on this node.");
+	}
 	return false;
 }
 
+UEdGraphPin* URigMapperDefinitionEditorGraphSchema::DropPinOnNode(UEdGraphNode* InTargetNode, const FName& InSourcePinName, const FEdGraphPinType& InSourcePinType, EEdGraphPinDirection InSourcePinDirection) const
+{
+	return Super::DropPinOnNode(InTargetNode, InSourcePinName, InSourcePinType, InSourcePinDirection);
+}
+
+void URigMapperDefinitionEditorGraphSchema::GetGraphContextActions(FGraphContextMenuBuilder& ContextMenuBuilder) const
+{
+	static const FText Category = LOCTEXT("RigMapperDefinitionEditor_AddNewGraphMenu", "Add new node");
+
+	const UEdGraphPin* FromPin = ContextMenuBuilder.FromPin;
+
+	if (!FromPin || FromPin->Direction == EGPD_Input)
+	{
+		ContextMenuBuilder.AddAction(MakeShared<FSchemaAction_SpawnRigMapperNode>(ERigMapperNodeType::Input, Category,
+			LOCTEXT("RigMapperDefinitionEditor_AddInputNode", "Add Input Node"),
+			LOCTEXT("RigMapperDefinitionEditor_AddInputTooltip", "Spawn Input node (at the end of the input connector, if available)"),
+			/*Grouping=*/0));
+	}
+
+	if (!FromPin || FromPin->Direction == EGPD_Output)
+	{
+		ContextMenuBuilder.AddAction(MakeShared<FSchemaAction_SpawnRigMapperNode>(ERigMapperNodeType::Output, Category,
+			LOCTEXT("RigMapperDefinitionEditor_AddOutputNode", "Add Output Node"),
+			LOCTEXT("RigMapperDefinitionEditor_AddOutputNodeTooltip", "Spawn Output node (at the end of the output connector, if available)"),
+			/*Grouping=*/0));
+	}
+
+	ContextMenuBuilder.AddAction(MakeShared<FSchemaAction_SpawnRigMapperNode>(ERigMapperNodeType::Multiply, Category,
+		LOCTEXT("RigMapperDefinitionEditor_AddMultiplyNode", "Add Multiply Node"),
+		LOCTEXT("RigMapperDefinitionEditor_AddMultiplyNodeTooltip", "Spawn Multiply node (at the end of the connector, if available)"),
+		/*Grouping=*/0));
+
+	ContextMenuBuilder.AddAction(MakeShared<FSchemaAction_SpawnRigMapperNode>(ERigMapperNodeType::SDK, Category,
+		LOCTEXT("RigMapperDefinitionEditor_AddSDKNode", "Add SDK Node"),
+		LOCTEXT("RigMapperDefinitionEditor_AddSDKNodeTooltip", "Spawn SDK node (at the end of the connector, if available)"),
+		/*Grouping=*/0));
+
+	ContextMenuBuilder.AddAction(MakeShared<FSchemaAction_SpawnRigMapperNode>(ERigMapperNodeType::WeightedSum, Category,
+		LOCTEXT("RigMapperDefinitionEditor_AddWeightedSumNode", "Add WeightedSum Node"),
+		LOCTEXT("RigMapperDefinitionEditor_AddWeightedSumNodeTooltip", "Spawn WeightedSum node (at the end of the connector, if available)"),
+		/*Grouping=*/0));
+
+	if (!FromPin)
+	{
+		ContextMenuBuilder.AddAction(MakeShared<FSchemaAction_SpawnRigMapperNode>(ERigMapperNodeType::NullOutput, Category,
+			LOCTEXT("RigMapperDefinitionEditor_AddNullOutputNode", "Add Null Output Node"),
+			LOCTEXT("RigMapperDefinitionEditor_AddNullOutputNodeTooltip", "Spawn null output node"),
+			/*Grouping=*/0));
+	}
+}
+
+bool URigMapperDefinitionEditorGraphSchema::TryCreateConnection(UEdGraphPin* InA, UEdGraphPin* InB) const
+{
+	if (InA == InB)
+	{
+		return false;
+	}
+	if (InA->GetOwningNode() == InB->GetOwningNode())
+	{
+		return false;
+	}
+
+	return Super::TryCreateConnection(InA, InB);
+}
+
+bool URigMapperDefinitionEditorGraphSchema::IsConnectionRelinkingAllowed(UEdGraphPin* InPin) const
+{
+	return Super::IsConnectionRelinkingAllowed(InPin);
+}
+
+const FPinConnectionResponse URigMapperDefinitionEditorGraphSchema::CanRelinkConnectionToPin(const UEdGraphPin* OldSourcePin, const UEdGraphPin* TargetPinCandidate) const
+{
+	return FPinConnectionResponse(CONNECT_RESPONSE_DISALLOW, TEXT(""));
+}
+
+bool URigMapperDefinitionEditorGraphSchema::CreatePromotedConnection(UEdGraphPin* A, UEdGraphPin* B) const
+{
+	return Super::CreatePromotedConnection(A, B);
+}
+
+void URigMapperDefinitionEditorGraphSchema::SetPinBeingDroppedOnNode(UEdGraphPin* InSourcePin) const
+{
+}
+
 #undef LOCTEXT_NAMESPACE
diff --git a/Engine/Plugins/Experimental/Animation/RigMapper/Source/RigMapperEditor/Private/RigMapperGraph/RigMapperDefinitionEditorGraphSchema.h b/Engine/Plugins/Experimental/Animation/RigMapper/Source/RigMapperEditor/Private/RigMapperGraph/RigMapperDefinitionEditorGraphSchema.h
--- a/Engine/Plugins/Experimental/Animation/RigMapper/Source/RigMapperEditor/Private/RigMapperGraph/RigMapperDefinitionEditorGraphSchema.h
+++ b/Engine/Plugins/Experimental/Animation/RigMapper/Source/RigMapperEditor/Private/RigMapperGraph/RigMapperDefinitionEditorGraphSchema.h
@@ -10,6 +10,24 @@
 
 #define UE_API RIGMAPPEREDITOR_API
 
+enum class ERigMapperNodeType : uint8;
+
+USTRUCT()
+struct FSchemaAction_SpawnRigMapperNode : public FEdGraphSchemaAction
+{
+	GENERATED_BODY()
+
+public:
+	FSchemaAction_SpawnRigMapperNode();
+	FSchemaAction_SpawnRigMapperNode(ERigMapperNodeType InNodeType, const FText& InCategory, const FText& InDescription, const FText& InTooltip, const int);
+
+	//~ Begin FEdGraphSchemaAction Interface
+	virtual UEdGraphNode* PerformAction(UEdGraph* ParentGraph, UEdGraphPin* FromPin, const FVector2f& Location, bool bSelectNewNode = true) override;
+	//~ End FEdGraphSchemaAction Interface
+private:
+	ERigMapperNodeType NodeType;
+};
+
 /**
  * 
  */
@@ -38,7 +56,7 @@ public:
 	};
 	
 	// UEdGraphSchema interface
-	UE_API virtual void GetContextMenuActions(class UToolMenu* Menu, class UGraphNodeContextMenuContext* Context) const override;
+	UE_API virtual void GetContextMenuActions(class UToolMenu* Menu, class UGraphNodeContextMenuContext* Context) const override;	
 	virtual FName GetParentContextMenuName() const override { return NAME_None; }
 	UE_API virtual FLinearColor GetPinTypeColor(const FEdGraphPinType& PinType) const override;
 	UE_API virtual void BreakPinLinks(UEdGraphPin& TargetPin, bool bSendsNodeNotification) const override;
@@ -48,9 +66,18 @@ public:
 	UE_API virtual class FConnectionDrawingPolicy* CreateConnectionDrawingPolicy(int32 InBackLayerID, int32 InFrontLayerID, float InZoomFactor, const FSlateRect& InClippingRect, class FSlateWindowElementList& InDrawElements, class UEdGraph* InGraphObj) const override;
 	virtual bool ShouldAlwaysPurgeOnModification() const override { return false; }
 	UE_API virtual FPinConnectionResponse CanCreateNewNodes(UEdGraphPin* InSourcePin) const override;
-	UE_API virtual const FPinConnectionResponse CanCreateConnection(const UEdGraphPin* A, const UEdGraphPin* B) const override;
+	UE_API virtual const FPinConnectionResponse CanCreateConnection(const UEdGraphPin* InA, const UEdGraphPin* InB) const override;
 	UE_API virtual bool SupportsDropPinOnNode(UEdGraphNode* InTargetNode, const FEdGraphPinType& InSourcePinType, EEdGraphPinDirection InSourcePinDirection, FText& OutErrorMessage) const override;
+	UE_API void GetGraphContextActions(FGraphContextMenuBuilder& ContextMenuBuilder) const override;
+	UE_API bool TryCreateConnection(UEdGraphPin* InA, UEdGraphPin* InB) const override;
+	UE_API bool IsConnectionRelinkingAllowed(UEdGraphPin* InPin) const override;
+	UE_API const FPinConnectionResponse CanRelinkConnectionToPin(const UEdGraphPin* OldSourcePin, const UEdGraphPin* TargetPinCandidate) const override;
+	UE_API bool CreatePromotedConnection(UEdGraphPin* InA, UEdGraphPin* InB) const override;
+	UEdGraphPin* DropPinOnNode(UEdGraphNode* InTargetNode, const FName& InSourcePinName, const FEdGraphPinType& InSourcePinType, EEdGraphPinDirection InSourcePinDirection) const override;
+	void SetPinBeingDroppedOnNode(UEdGraphPin* InSourcePin) const override;
 	// End of UEdGraphSchema interface
+private:
+
 };
 
 #undef UE_API
diff --git a/Engine/Plugins/Experimental/Animation/RigMapper/Source/RigMapperEditor/Private/RigMapperGraph/SRigMapperDefinitionGraphEditor.cpp b/Engine/Plugins/Experimental/Animation/RigMapper/Source/RigMapperEditor/Private/RigMapperGraph/SRigMapperDefinitionGraphEditor.cpp
--- a/Engine/Plugins/Experimental/Animation/RigMapper/Source/RigMapperEditor/Private/RigMapperGraph/SRigMapperDefinitionGraphEditor.cpp
+++ b/Engine/Plugins/Experimental/Animation/RigMapper/Source/RigMapperEditor/Private/RigMapperGraph/SRigMapperDefinitionGraphEditor.cpp
@@ -33,12 +33,13 @@ void SRigMapperDefinitionGraphEditor::Construct(const FArguments& InArgs, URigMa
 	SGraphEditor::FGraphEditorEvents GraphEvents;
 	// GraphEvents.OnCreateActionMenu = SGraphEditor::FOnCreateActionMenu::CreateSP(this, &SPhysicsAssetGraph::OnCreateGraphActionMenu);
 	GraphEvents.OnSelectionChanged = SGraphEditor::FOnSelectionChanged::CreateSP(this, &SRigMapperDefinitionGraphEditor::HandleSelectionChanged);
-	// GraphEvents.OnNodeDoubleClicked = FSingleNodeEvent::CreateSP(this, &SPhysicsAssetGraph::HandleNodeDoubleClicked);
+	// GraphEvents.OnNodeDoubleClicked = FSingleNodeEvent::CreateSP(this, &SPhysicsAssetGraph::HandleNodeDoubleClicked);	
 
 	GraphObj = NewObject<URigMapperDefinitionEditorGraph>();
 	GraphObj->AddToRoot();
 	GraphObj->Initialize(InDefinition);
 	GraphObj->RebuildGraph();
+	GraphObj->LayoutNodes();
 
 	ChildSlot
 	[
@@ -51,12 +52,11 @@ void SRigMapperDefinitionGraphEditor::Construct(const FArguments& InArgs, URigMa
 			.GraphEvents(GraphEvents)
 			.Appearance(AppearanceInfo)
 			.ShowGraphStateOverlay(false)
-			.IsEditable(false)
-			.AutoExpandActionMenu(false)
-			.DisplayAsReadOnly(false)
+			.IsEditable(true)
+			.AutoExpandActionMenu(true)
+			.DisplayAsReadOnly(false)				
 		]
 	];
-
 }
 
 void SRigMapperDefinitionGraphEditor::Tick(const FGeometry& AllottedGeometry, const double InCurrentTime, const float InDeltaTime)
@@ -68,10 +68,10 @@ void SRigMapperDefinitionGraphEditor::Tick(const FGeometry& AllottedGeometry, co
 		GraphObj->RequestRefreshLayout(false);
 
 		const FGraphPanelSelectionSet& Selection = GraphEditor->GetSelectedNodes();
-		
+
 		TArray<URigMapperDefinitionEditorGraphNode*> SelectedNodes;
 		SelectedNodes.Reserve(Selection.Num());
-		
+
 		for (UObject* Obj : Selection)
 		{
 			if (URigMapperDefinitionEditorGraphNode* Node = Cast<URigMapperDefinitionEditorGraphNode>(Obj))
@@ -79,7 +79,7 @@ void SRigMapperDefinitionGraphEditor::Tick(const FGeometry& AllottedGeometry, co
 				SelectedNodes.Add(Node);
 			}
 		}
-		
+
 		ZoomToFitNodes(SelectedNodes);
 	}
 }
@@ -98,7 +98,7 @@ void SRigMapperDefinitionGraphEditor::SelectNodes(const TArray<FString>& Inputs,
 		{
 			GraphEditor->SetNodeSelection(Node, true);
 		}
-		ZoomToFitNodes(SelectedNodes);
+		//ZoomToFitNodes(SelectedNodes);
 		
 		bSelectingNodes = false;
 	}
@@ -109,6 +109,15 @@ void SRigMapperDefinitionGraphEditor::RebuildGraph()
 	GraphObj->RebuildGraph();
 }
 
+void SRigMapperDefinitionGraphEditor::RefreshGraphNode(URigMapperDefinitionEditorGraphNode* InNode)
+{
+	if (InNode)
+	{
+		InNode->Modify();
+		GraphEditor->RefreshNode(*InNode);
+	}
+}
+
 void SRigMapperDefinitionGraphEditor::ZoomToFitNodes(const TArray<URigMapperDefinitionEditorGraphNode*>& SelectedNodes) const
 {
 	if (GraphObj->NeedsRefreshLayout())
@@ -215,7 +224,7 @@ void SRigMapperDefinitionGraphEditor::HandleSelectionChanged(const TSet<UObject*
 			}
 		}
 		
-		ZoomToFitNodes(SelectedNodes);
+		//ZoomToFitNodes(SelectedNodes);
 		
 		bSelectingNodes = false;
 	}
diff --git a/Engine/Plugins/Experimental/Animation/RigMapper/Source/RigMapperEditor/Private/RigMapperGraph/SRigMapperDefinitionGraphEditor.h b/Engine/Plugins/Experimental/Animation/RigMapper/Source/RigMapperEditor/Private/RigMapperGraph/SRigMapperDefinitionGraphEditor.h
--- a/Engine/Plugins/Experimental/Animation/RigMapper/Source/RigMapperEditor/Private/RigMapperGraph/SRigMapperDefinitionGraphEditor.h
+++ b/Engine/Plugins/Experimental/Animation/RigMapper/Source/RigMapperEditor/Private/RigMapperGraph/SRigMapperDefinitionGraphEditor.h
@@ -33,6 +33,7 @@ public:
 
 	UE_API void SelectNodes(const TArray<FString>& Inputs, const TArray<FString>& Features, const TArray<FString>& Outputs, const TArray<FString>& NullOutputs);
 	UE_API void RebuildGraph();
+	UE_API void RefreshGraphNode(URigMapperDefinitionEditorGraphNode* InNode);
 
 private:
 	UE_API void ZoomToFitNodes(const TArray<URigMapperDefinitionEditorGraphNode*>& SelectedNodes) const;
diff --git a/Engine/Plugins/Experimental/Animation/RigMapper/Source/RigMapperEditor/Private/RigMapperGraph/SRigMapperDefinitionGraphEditorNode.cpp b/Engine/Plugins/Experimental/Animation/RigMapper/Source/RigMapperEditor/Private/RigMapperGraph/SRigMapperDefinitionGraphEditorNode.cpp
--- a/Engine/Plugins/Experimental/Animation/RigMapper/Source/RigMapperEditor/Private/RigMapperGraph/SRigMapperDefinitionGraphEditorNode.cpp
+++ b/Engine/Plugins/Experimental/Animation/RigMapper/Source/RigMapperEditor/Private/RigMapperGraph/SRigMapperDefinitionGraphEditorNode.cpp
@@ -5,6 +5,7 @@
 #include "RigMapperDefinitionEditorGraphNode.h"
 #include "SlateOptMacros.h"
 #include "Components/HorizontalBox.h"
+#include "SGraphPin.h"
 
 BEGIN_SLATE_FUNCTION_BUILD_OPTIMIZATION
 
@@ -21,6 +22,15 @@ void SRigMapperDefinitionGraphEditorNode::Construct(const FArguments& InArgs, UR
 	UpdateGraphNode();
 }
 
+void SRigMapperDefinitionGraphEditorNode::CreatePinWidgets()
+{
+	for (UEdGraphPin* Pin : GraphNode->Pins)
+	{
+		TSharedPtr<SGraphPin> PinWidget = CreatePinWidget(Pin);
+		AddPin(PinWidget.ToSharedRef());	
+	}
+}
+
 FSlateColor SRigMapperDefinitionGraphEditorNode::GetNodeColor() const
 {
 	if (GraphNode)
@@ -69,7 +79,7 @@ void SRigMapperDefinitionGraphEditorNode::UpdateGraphNode()
 			.Padding(0)
 			[
 				SNew(SHorizontalBox)
-
+				//.Visibility(EVisibility::HitTestInvisible)
 				+SHorizontalBox::Slot()
 				.AutoWidth()
 				.HAlign(HAlign_Center)
@@ -120,6 +130,60 @@ void SRigMapperDefinitionGraphEditorNode::UpdateGraphNode()
 	CreatePinWidgets();
 }
 
+TSharedPtr<SGraphPin> SRigMapperDefinitionGraphEditorNode::CreatePinWidget(UEdGraphPin* Pin) const
+{
+	// TODO: Create our own NodeFactory create pin widget method.
+	return SNew(SGraphPin, Pin);
+}
+
+void SRigMapperDefinitionGraphEditorNode::AddPin(const TSharedRef<SGraphPin>& PinToAdd)
+{
+	PinToAdd->SetOwner(SharedThis(this));
+	PinToAdd->SetIsEditable(true);
+
+	const UEdGraphPin* PinObj = PinToAdd->GetPinObj();
+	const bool bAdvancedParameter = PinObj && PinObj->bAdvancedView;
+	if (bAdvancedParameter)
+	{
+		PinToAdd->SetVisibility(TAttribute<EVisibility>(PinToAdd, &SGraphPin::IsPinVisibleAsAdvanced));
+	}
+
+	if (PinToAdd->GetDirection() == EEdGraphPinDirection::EGPD_Input)
+	{
+		LeftNodeBox->AddSlot()
+			.AutoHeight()
+			.HAlign(HAlign_Left)
+			.VAlign(VAlign_Center)
+			.FillHeight(1.0f)
+			[
+				SNew(SOverlay)
+				.Visibility(EVisibility::SelfHitTestInvisible)
+				+ SOverlay::Slot()
+				[
+					PinToAdd
+				]
+			];
+		InputPins.Add(PinToAdd);
+	}
+	else // Direction == EEdGraphPinDirection::EGPD_Output
+	{
+		RightNodeBox->AddSlot()
+			.AutoHeight()
+			.HAlign(HAlign_Right)
+			.VAlign(VAlign_Center)
+			.FillHeight(1.0f)
+			[
+				SNew(SOverlay)
+				.Visibility(EVisibility::SelfHitTestInvisible)
+				+ SOverlay::Slot()
+				[
+					PinToAdd
+				]
+			];
+		OutputPins.Add(PinToAdd);
+	}
+}
+
 const FSlateBrush* SRigMapperDefinitionGraphEditorNode::GetShadowBrush(bool bSelected) const
 {
 	return bSelected ? FAppStyle::GetBrush(TEXT("PhysicsAssetEditor.Graph.Node.ShadowSelected")) : FAppStyle::GetBrush(TEXT("PhysicsAssetEditor.Graph.Node.Shadow"));
diff --git a/Engine/Plugins/Experimental/Animation/RigMapper/Source/RigMapperEditor/Private/RigMapperGraph/SRigMapperDefinitionGraphEditorNode.h b/Engine/Plugins/Experimental/Animation/RigMapper/Source/RigMapperEditor/Private/RigMapperGraph/SRigMapperDefinitionGraphEditorNode.h
--- a/Engine/Plugins/Experimental/Animation/RigMapper/Source/RigMapperEditor/Private/RigMapperGraph/SRigMapperDefinitionGraphEditorNode.h
+++ b/Engine/Plugins/Experimental/Animation/RigMapper/Source/RigMapperEditor/Private/RigMapperGraph/SRigMapperDefinitionGraphEditorNode.h
@@ -23,6 +23,7 @@ public:
 
 	/** Constructs this widget with InArgs */
 	UE_API void Construct(const FArguments& InArgs, URigMapperDefinitionEditorGraphNode* InNode);
+	UE_API void CreatePinWidgets() override;
 
 protected:
 	UE_API FSlateColor GetNodeColor() const;
@@ -33,6 +34,8 @@ protected:
 	// SGraphNode interface
 	UE_API virtual void UpdateGraphNode() override;
 	UE_API virtual const FSlateBrush* GetShadowBrush(bool bSelected) const override;
+	UE_API TSharedPtr<SGraphPin> CreatePinWidget(UEdGraphPin* Pin) const override;
+	UE_API void AddPin(const TSharedRef<SGraphPin>& PinToAdd) override;
 };
 
 #undef UE_API
