[RigMapper] Quality of Life improvements:
- Adding more pins to nodes and generating connection
- Fixing bug with loading json file
- Added constraints when removing pins
- Improving memory management and performance

diff --git a/Engine/Plugins/Experimental/Animation/RigMapper/Source/RigMapper/Private/RigMapperDefinition.cpp b/Engine/Plugins/Experimental/Animation/RigMapper/Source/RigMapper/Private/RigMapperDefinition.cpp
--- a/Engine/Plugins/Experimental/Animation/RigMapper/Source/RigMapper/Private/RigMapperDefinition.cpp
+++ b/Engine/Plugins/Experimental/Animation/RigMapper/Source/RigMapper/Private/RigMapperDefinition.cpp
@@ -180,8 +180,19 @@ bool FRigMapperWsFeature::IsValid(const TArray<FString>& InputNames, bool bWarn)
 }
 
 void FRigMapperWsFeature::GetInputs(TArray<FString>& OutInputs) const
-{
-	Inputs.GenerateKeyArray(OutInputs);
+{	
+	if (!OrderedInputs.IsEmpty())
+	{
+		OutInputs.Empty(OrderedInputs.Num());
+		for (int32 i = 0; i < OrderedInputs.Num(); ++i)
+		{
+			OutInputs.Add(OrderedInputs[i].Name);
+		}
+	}
+	else
+	{
+		Inputs.GenerateKeyArray(OutInputs);
+	}
 }
 
 bool FRigMapperSdkFeature::LoadFromJsonObject(const TSharedPtr<FJsonObject>& JsonObject)
@@ -684,6 +695,34 @@ bool URigMapperDefinition::WasDefinitionValidated() const
 	return bValidated;
 }
 
+FRigMapperFeature* URigMapperDefinition::GetFeature(const FString& InFeatureName, const ERigMapperFeatureType InFeatureType)
+{
+	switch (InFeatureType)
+	{
+	case ERigMapperFeatureType::Multiply:
+		return Features.Multiply.FindByPredicate([InFeatureName](const FRigMapperMultiplyFeature& Other) { return Other.Name == InFeatureName; });
+
+	case ERigMapperFeatureType::WeightedSum:
+		return Features.WeightedSums.FindByPredicate([InFeatureName](const FRigMapperWsFeature& Other) { return Other.Name == InFeatureName; });
+
+	case ERigMapperFeatureType::SDK:
+		return Features.SDKs.FindByPredicate([InFeatureName](const FRigMapperSdkFeature& Other) { return Other.Name == InFeatureName; });
+
+	case ERigMapperFeatureType::Input:
+		return Features.Inputs.FindByPredicate([InFeatureName](const FRigMapperFeature& Other) { return Other.Name == InFeatureName; });
+
+	case ERigMapperFeatureType::Output:
+		return Features.Outputs.FindByPredicate([InFeatureName](const FRigMapperOutputFeature& Other) { return Other.Name == InFeatureName; });
+
+	case ERigMapperFeatureType::NullOutput:
+		return Features.NullOutputs.FindByPredicate([InFeatureName](const FRigMapperNullOutputFeature& Other) { return Other.Name == InFeatureName; });
+
+	default:
+		break;
+	}
+	return nullptr;
+}
+
 bool URigMapperDefinition::IsDefinitionValid(bool bWarn, bool bForce) const
 {
 	if (!bForce && WasDefinitionValidated())
@@ -856,6 +895,22 @@ bool FRigMapperMultiplyFeature::BakeInput(URigMapperLinkedDefinitions* LinkedDef
 	return OutBakedInput.Key.IsValid() && !bMissingIsNotNullOutput; // if we have a missing definition which is not defined as a NullOutput, flag a failure
 }
 
+void FRigMapperMultiplyFeature::RelinkFeature(const FString& InOldInput, const FString& InNewInput)
+{
+	for (int32 i = 0; i < Inputs.Num(); i++)
+	{
+		if (Inputs[i] == InOldInput)
+		{
+			Inputs[i] = InNewInput;
+		}
+	}
+}
+
+void FRigMapperMultiplyFeature::AddNewInputLink(const FString& InNewInput)
+{
+	Inputs.Add(InNewInput);
+}
+
 bool FRigMapperWsFeature::BakeInput(URigMapperLinkedDefinitions* LinkedDefinitions, int32 DefinitionIndex, FBakedInput& OutBakedInput) const
 {
 	const FString BakedInputName = FString::Printf(TEXT("%s:%d"), *Name, DefinitionIndex);
@@ -900,6 +955,72 @@ bool FRigMapperWsFeature::BakeInput(URigMapperLinkedDefinitions* LinkedDefinitio
 	return OutBakedInput.Key.IsValid() && !bMissingIsNotNullOutput; // if we have a missing definition which is not defined as a NullOutput, flag a failure
 }
 
+void FRigMapperWsFeature::RelinkFeature(const FString& InOldInput, const FString& InNewInput)
+{
+	if (double* Weight = Inputs.Find(InOldInput))
+	{
+		for (int32 i = 0; i < OrderedInputs.Num(); i++)
+		{
+			if (OrderedInputs[i].Name == InOldInput)
+			{
+				OrderedInputs[i].Name = InNewInput; // Preserves the pin ordering.
+			}
+		}		
+		double TempValue = MoveTemp(*Weight);
+		Inputs.Remove(InOldInput);
+		Inputs.Add(InNewInput, TempValue);
+	}
+}
+
+void FRigMapperWsFeature::AddNewInputLink(const FString& InNewInput)
+{
+	if (!Inputs.Contains(InNewInput))
+	{
+		Inputs.Add(InNewInput, 0.0f);
+		OrderedInputs.Add(FRigMapperWsInput(InNewInput, 0.0f));
+	}
+}
+
+void FRigMapperWsFeature::InitializeDetails()
+{
+	if (OrderedInputs.IsEmpty())
+	{
+		for (const TPair<FString, double>& Pair : Inputs)
+		{
+			OrderedInputs.Add(FRigMapperWsInput(Pair.Key, Pair.Value));
+		}
+	}
+}
+
+void FRigMapperWsFeature::UpdateFromDetails()
+{
+	Inputs.Empty();
+	for (FRigMapperWsInput& WsInput : OrderedInputs)
+	{
+		Inputs.Add(WsInput.Name, WsInput.Weight);
+	}
+}
+
+void FRigMapperWsFeature::RemoveInput(const FString& InInput)
+{
+	Inputs.Remove(InInput);
+}
+
+void FRigMapperWsFeature::ClearInputs()
+{
+	Inputs.Empty();
+}
+
+void FRigMapperWsFeature::RenameInput(const FString& InOldInput, const FString& InNewInput)
+{
+	if (double* Weight = Inputs.Find(InOldInput))
+	{
+		double TempValue = MoveTemp(*Weight);
+		Inputs.Remove(InOldInput);
+		Inputs.Add(InNewInput, TempValue);
+	}
+}
+
 bool FRigMapperSdkFeature::BakeInput(URigMapperLinkedDefinitions* LinkedDefinitions, int32 DefinitionIndex, FBakedInput& OutBakedInput) const
 {
 	const FString BakedInputName = FString::Printf(TEXT("%s:%d"), *Name, DefinitionIndex);
@@ -932,6 +1053,16 @@ bool FRigMapperSdkFeature::BakeInput(URigMapperLinkedDefinitions* LinkedDefiniti
 	return OutBakedInput.Key.IsValid() && !bMissingIsNotNullOutput; // if we have a missing definition which is not defined as a NullOutput, flag a failure
 }
 
+void FRigMapperSdkFeature::RelinkFeature(const FString& InOldInput, const FString& InNewInput)
+{
+	Input = InNewInput;
+}
+
+void FRigMapperSdkFeature::AddNewInputLink(const FString& InNewInput)
+{
+	Input = InNewInput;
+}
+
 bool FRigMapperSdkFeature::BakeKeys(const FRigMapperSdkFeature& InSdk, const FRigMapperSdkFeature& OutSdk, TArray<FRigMapperSdkKey>& BakedKeys)
 {
 	// Create an array to backward evaluate InSdk (First Layer)
@@ -1187,3 +1318,13 @@ bool URigMapperLinkedDefinitions::AreLinkedDefinitionsValid() const
 	
 	return bOk;
 }
+
+void FRigMapperOutputFeature::RelinkFeature(const FString& InOldInput, const FString& InNewInput)
+{
+	Input = InNewInput;
+}
+
+void FRigMapperOutputFeature::AddNewInputLink(const FString& InNewInput)
+{
+	Input = InNewInput;
+}
diff --git a/Engine/Plugins/Experimental/Animation/RigMapper/Source/RigMapper/Public/RigMapperDefinition.h b/Engine/Plugins/Experimental/Animation/RigMapper/Source/RigMapper/Public/RigMapperDefinition.h
--- a/Engine/Plugins/Experimental/Animation/RigMapper/Source/RigMapper/Public/RigMapperDefinition.h
+++ b/Engine/Plugins/Experimental/Animation/RigMapper/Source/RigMapper/Public/RigMapperDefinition.h
@@ -23,7 +23,8 @@ enum class ERigMapperFeatureType : uint8
 	SDK,
 	Multiply,
 	Output,
-	NullOutput
+	NullOutput,
+	Invalid
 };
 
 USTRUCT(BlueprintType)
@@ -50,10 +51,23 @@ public:
 
 	virtual ERigMapperFeatureType GetFeatureType() const { return ERigMapperFeatureType::Input; }
 
-	virtual void GetInputs(TArray<FString>& OutInputs) const { }
+	virtual void GetInputs(TArray<FString>& OutInputs) const {}
 
 	virtual bool BakeInput(URigMapperLinkedDefinitions* LinkedDefinitions, int32 DefinitionIndex, FBakedInput& OutBakedInput) const { return false; }
 
+	virtual void RelinkFeature(const FString& InOldInput, const FString& InNewInput) {}
+
+	virtual void AddNewInputLink(const FString& InNewInput) {}
+
+	/** Initializes additional non-serialized properties used for details view */
+	virtual void InitializeDetails() {}
+
+	/** Update feature properties from details view */
+	virtual void UpdateFromDetails() {}
+
+	virtual void RemoveInput(const FString& InInput) {};
+	virtual void ClearInputs() {};
+	virtual void RenameInput(const FString& InOldInput, const FString& InNewInput) {};
 
 	bool operator==(const FRigMapperFeature& InOther) const
 	{
@@ -89,6 +103,10 @@ public:
 
 	UE_API virtual bool BakeInput(URigMapperLinkedDefinitions* LinkedDefinitions, int32 DefinitionIndex, FBakedInput& OutBakedInput) const override;
 
+	UE_API virtual void RelinkFeature(const FString& InOldInput, const FString& InNewInput) override;
+	
+	UE_API virtual void AddNewInputLink(const FString& InNewInput) override;
+
 	bool operator==(const FRigMapperMultiplyFeature& InOther) const
 	{
 		return FRigMapperFeature::operator==(InOther) && Inputs == InOther.Inputs;
@@ -120,15 +138,36 @@ public:
 	}
 };
 
+USTRUCT(BlueprintType)
+struct FRigMapperWsInput
+{
+	GENERATED_BODY()
+
+public:
+	FRigMapperWsInput() {}
+
+	explicit FRigMapperWsInput(const FString& InName, const double InWeight) : Name(InName), Weight(InWeight) {}
+
+public:
+	UPROPERTY(EditAnywhere, BlueprintReadWrite, Category = "Animation|Rig Mapper")
+	FString Name;
+
+	UPROPERTY(EditAnywhere, BlueprintReadWrite, Category = "Animation|Rig Mapper")
+	double Weight;
+};
+
 USTRUCT(BlueprintType)
 struct FRigMapperWsFeature : public FRigMapperFeature
 {
 	GENERATED_BODY()
 
 public:
-	UPROPERTY(EditAnywhere, BlueprintReadWrite, Category="Animation|Rig Mapper", meta=(DisplayPriority=1))
+	UPROPERTY()
 	TMap<FString, double> Inputs;
 
+	UPROPERTY(EditAnywhere, BlueprintReadWrite, Category="Animation|Rig Mapper", meta=(DisplayPriority=1))
+	TArray<FRigMapperWsInput> OrderedInputs;
+
 	UPROPERTY(EditAnywhere, BlueprintReadWrite, Category="Animation|Rig Mapper", meta=(DisplayPriority=1))
 	FRigMapperFeatureRange Range;
 	
@@ -147,6 +186,18 @@ public:
 
 	UE_API virtual bool BakeInput(URigMapperLinkedDefinitions* LinkedDefinitions, int32 DefinitionIndex, FBakedInput& OutBakedInput) const override;
 
+	UE_API virtual void RelinkFeature(const FString& InOldInput, const FString& InNewInput) override;
+
+	UE_API virtual void AddNewInputLink(const FString& InNewInput) override;
+
+	UE_API virtual void InitializeDetails() override;
+
+	UE_API virtual void UpdateFromDetails() override;
+
+	UE_API virtual void RemoveInput(const FString& InInput) override;
+	UE_API virtual void ClearInputs() override;
+	UE_API virtual void RenameInput(const FString& InOldInput, const FString& InNewInput) override;
+
 	bool operator==(const FRigMapperWsFeature& InOther) const
 	{
 		if (Inputs.Num() != InOther.Inputs.Num())
@@ -212,6 +263,10 @@ public:
 
 	UE_API virtual bool BakeInput(URigMapperLinkedDefinitions* LinkedDefinitions, int32 DefinitionIndex, FBakedInput& OutBakedInput) const override;
 
+	UE_API virtual void RelinkFeature(const FString& InOldInput, const FString& InNewInput) override;
+
+	UE_API virtual void AddNewInputLink(const FString& InNewInput) override;
+
 	static UE_API bool BakeKeys(const FRigMapperSdkFeature& InSdk, const FRigMapperSdkFeature& OutSdk, TArray<FRigMapperSdkKey>& BakedKeys);
 
 	bool operator==(const FRigMapperSdkFeature& InOther) const
@@ -238,6 +293,10 @@ public:
 
 	virtual void GetInputs(TArray<FString>& OutInputs) const override { OutInputs = { Input }; };
 
+	UE_API virtual void RelinkFeature(const FString& InOldInput, const FString& InNewInput) override;
+
+	UE_API virtual void AddNewInputLink(const FString& InNewInput) override;
+
 	bool operator==(const FRigMapperOutputFeature& InOther) const
 	{
 		return FRigMapperFeature::operator==(InOther) && Input == InOther.Input;
@@ -255,6 +314,10 @@ public:
 	explicit FRigMapperNullOutputFeature(const FString& InName) : FRigMapperFeature(InName) { }
 
 	virtual ERigMapperFeatureType GetFeatureType() const override { return ERigMapperFeatureType::NullOutput; }
+
+	UE_API virtual void RelinkFeature(const FString& InOldInput, const FString& InNewInput) override {};
+
+	UE_API virtual void AddNewInputLink(const FString& InNewInput) override {};
 };
 
 USTRUCT(BlueprintType)
@@ -271,6 +334,16 @@ struct FRigMapperFeatureDefinitions
 	UPROPERTY(EditAnywhere, BlueprintReadWrite, Category="Animation|Rig Mapper")
 	TArray<FRigMapperSdkFeature> SDKs;
 
+	/** RigMapper Features that are not serialized, but created on startup from strings in the Definition */
+	UPROPERTY(Transient)
+	TArray<FRigMapperFeature> Inputs;
+
+	UPROPERTY(Transient)
+	TArray<FRigMapperOutputFeature> Outputs;
+
+	UPROPERTY(Transient)
+	TArray<FRigMapperNullOutputFeature> NullOutputs;
+
 public:
 	UE_API bool AddFromJsonObject(const FString& FeatureName, const TSharedPtr<FJsonObject>& JsonObject);
 
@@ -344,6 +417,11 @@ public:
 	UFUNCTION(BlueprintCallable, Category = "Animation|Rig Mapper")
 	UE_API bool WasDefinitionValidated() const;
 
+	/**
+	 * Gets the Feature by name and type, including Inputs, Outputs and NullOutputs.
+	 */
+	UE_API FRigMapperFeature* GetFeature(const FString& InFeatureName, const ERigMapperFeatureType InFeatureType);
+
 	
 private:
 	bool LoadInputsFromJsonObject(const TSharedPtr<FJsonObject>& JsonObject);
diff --git a/Engine/Plugins/Experimental/Animation/RigMapper/Source/RigMapperEditor/Private/RigMapperDefinitionEditorToolkit.cpp b/Engine/Plugins/Experimental/Animation/RigMapper/Source/RigMapperEditor/Private/RigMapperDefinitionEditorToolkit.cpp
--- a/Engine/Plugins/Experimental/Animation/RigMapper/Source/RigMapperEditor/Private/RigMapperDefinitionEditorToolkit.cpp
+++ b/Engine/Plugins/Experimental/Animation/RigMapper/Source/RigMapperEditor/Private/RigMapperDefinitionEditorToolkit.cpp
@@ -22,20 +22,23 @@ const FName FRigMapperDefinitionEditorToolkit::DefinitionEditorGraphTabId("RigMa
 const FName FRigMapperDefinitionEditorToolkit::DefinitionEditorStructureTabId("RigMapperEditor_DefinitionStructureView");
 const FName FRigMapperDefinitionEditorToolkit::DefinitionEditorDetailsTabId("RigMapperEditor_DefinitionDetailsView");
 
-const TMap<FName, ERigMapperNodeType> FRigMapperDefinitionEditorToolkit::PropertyNameToNodeTypeMapping = {
-	{ GET_MEMBER_NAME_CHECKED(URigMapperDefinition, Inputs), ERigMapperNodeType::Input },
-	{ GET_MEMBER_NAME_CHECKED(URigMapperDefinition, Features), ERigMapperNodeType::Invalid },
-	{ GET_MEMBER_NAME_CHECKED(FRigMapperFeatureDefinitions, Multiply), ERigMapperNodeType::Multiply },
-	{ GET_MEMBER_NAME_CHECKED(FRigMapperFeatureDefinitions, WeightedSums), ERigMapperNodeType::WeightedSum },
-	{ GET_MEMBER_NAME_CHECKED(FRigMapperFeatureDefinitions, SDKs), ERigMapperNodeType::SDK },
-	{ GET_MEMBER_NAME_CHECKED(URigMapperDefinition, Outputs), ERigMapperNodeType::Output },
-	{ GET_MEMBER_NAME_CHECKED(URigMapperDefinition, NullOutputs), ERigMapperNodeType::NullOutput },
+const TMap<FName, ERigMapperFeatureType> FRigMapperDefinitionEditorToolkit::PropertyNameToNodeTypeMapping = {
+	{ GET_MEMBER_NAME_CHECKED(URigMapperDefinition, Inputs), ERigMapperFeatureType::Input },
+	{ GET_MEMBER_NAME_CHECKED(URigMapperDefinition, Features), ERigMapperFeatureType::Invalid },
+	{ GET_MEMBER_NAME_CHECKED(FRigMapperFeatureDefinitions, Multiply), ERigMapperFeatureType::Multiply },
+	{ GET_MEMBER_NAME_CHECKED(FRigMapperFeatureDefinitions, WeightedSums), ERigMapperFeatureType::WeightedSum },
+	{ GET_MEMBER_NAME_CHECKED(FRigMapperFeatureDefinitions, SDKs), ERigMapperFeatureType::SDK },
+	{ GET_MEMBER_NAME_CHECKED(URigMapperDefinition, Outputs), ERigMapperFeatureType::Output },
+	{ GET_MEMBER_NAME_CHECKED(URigMapperDefinition, NullOutputs), ERigMapperFeatureType::NullOutput },
 };
 
 void FRigMapperDefinitionEditorToolkit::Initialize(URigMapperDefinition* InDefinition, EToolkitMode::Type InMode, TSharedPtr<IToolkitHost> InToolkitHost)
 {
 	Definition = InDefinition;
-
+	for (FRigMapperFeature& Feature : Definition->Features.WeightedSums)
+	{
+		Feature.InitializeDetails();
+	}	
 	// todo
 	// GEditor->GetEditorSubsystem<UImportSubsystem>()->OnAssetPostImport.AddSP(this, &FSimpleAssetEditor::HandleAssetPostImport);
 	// FCoreUObjectDelegates::OnObjectsReplaced.AddSP(this, &FSimpleAssetEditor::OnObjectsReplaced);
@@ -155,12 +158,20 @@ void FRigMapperDefinitionEditorToolkit::HandleRigMapperDefinitionLoaded()
 {
 	if (GraphEditor)
 	{
-		//GraphEditor->RebuildGraph();
+		GraphEditor->RebuildGraph();
 	}
 	if (StructureView)
 	{
 		StructureView->RebuildTree();
 	}
+	if (FeatureDetailsView)
+	{
+		FeatureDetailsView->SetStructureData(nullptr);
+	}
+	for (FRigMapperFeature& Feature : Definition->Features.WeightedSums)
+	{
+		Feature.InitializeDetails();
+	}	
 }
 
 TSharedRef<IDetailCustomization> FRigMapperDefinitionEditorToolkit::FDetailsViewCustomization::MakeInstance()
@@ -289,38 +300,47 @@ void FRigMapperDefinitionEditorToolkit::HandleFinishedChangingProperties(const F
 		if (bInput)
 		{
 			int32 InputIndex = Event.GetArrayIndex(PropertyName.ToString());
-			StructureView->SelectNode(Definition->Inputs[InputIndex], ERigMapperNodeType::Input, true);
-			GraphEditor->SelectNodes({ Definition->Inputs[InputIndex] }, {}, {}, {});
+			if (Definition->Inputs.IsValidIndex(InputIndex))
+			{
+				StructureView->SelectNode(Definition->Inputs[InputIndex], ERigMapperFeatureType::Input, true);
+				GraphEditor->SelectNodes({ Definition->Inputs[InputIndex] }, {}, {}, {});
+			}
 		}
 		else if (bOutput)
 		{
 			int32 OutputIndex = Event.GetArrayIndex(PropertyName.ToString());
 			TArray<FString> OutputsInOrder;
 			Definition->Outputs.GenerateKeyArray(OutputsInOrder);
-			StructureView->SelectNode(OutputsInOrder[OutputIndex], ERigMapperNodeType::Output, true);
-			GraphEditor->SelectNodes({}, {}, { OutputsInOrder[OutputIndex] }, {});
+			if (OutputsInOrder.IsValidIndex(OutputIndex))
+			{
+				StructureView->SelectNode(OutputsInOrder[OutputIndex], ERigMapperFeatureType::Output, true);
+				GraphEditor->SelectNodes({}, {}, { OutputsInOrder[OutputIndex] }, {});
+			}
 		}
 		if (bNullOutput)
 		{
 			int32 NullOutputIndex = Event.GetArrayIndex(PropertyName.ToString());
-			StructureView->SelectNode(Definition->NullOutputs[NullOutputIndex], ERigMapperNodeType::NullOutput, true);
-			GraphEditor->SelectNodes({ Definition->NullOutputs[NullOutputIndex] }, {}, {}, {});
+			if (Definition->NullOutputs.IsValidIndex(NullOutputIndex))
+			{
+				StructureView->SelectNode(Definition->NullOutputs[NullOutputIndex], ERigMapperFeatureType::NullOutput, true);
+				GraphEditor->SelectNodes({ Definition->NullOutputs[NullOutputIndex] }, {}, {}, {});
+			}
 		}
 		else if (bFeaturesChildProperty)
 		{
 			if (const int32 MultiplyIndex = Event.GetArrayIndex(GET_MEMBER_NAME_CHECKED(FRigMapperFeatureDefinitions, Multiply).ToString()); MultiplyIndex != INDEX_NONE)
 			{
-				StructureView->SelectNode(Definition->Features.Multiply[MultiplyIndex].Name, ERigMapperNodeType::Multiply, true);
+				StructureView->SelectNode(Definition->Features.Multiply[MultiplyIndex].Name, ERigMapperFeatureType::Multiply, true);
 				GraphEditor->SelectNodes({}, { Definition->Features.Multiply[MultiplyIndex].Name }, {}, {});
 			}
 			else if (const int32 WsIndex = Event.GetArrayIndex(GET_MEMBER_NAME_CHECKED(FRigMapperFeatureDefinitions, WeightedSums).ToString()); WsIndex != INDEX_NONE)
 			{
-				StructureView->SelectNode(Definition->Features.WeightedSums[WsIndex].Name, ERigMapperNodeType::WeightedSum, true);
+				StructureView->SelectNode(Definition->Features.WeightedSums[WsIndex].Name, ERigMapperFeatureType::WeightedSum, true);
 				GraphEditor->SelectNodes({}, { Definition->Features.WeightedSums[WsIndex].Name }, { }, {});
 			}
 			else if (const int32 SdkIndex = Event.GetArrayIndex(GET_MEMBER_NAME_CHECKED(FRigMapperFeatureDefinitions, SDKs).ToString()); SdkIndex != INDEX_NONE)
 			{
-				StructureView->SelectNode(Definition->Features.SDKs[SdkIndex].Name, ERigMapperNodeType::SDK, true);
+				StructureView->SelectNode(Definition->Features.SDKs[SdkIndex].Name, ERigMapperFeatureType::SDK, true);
 				GraphEditor->SelectNodes({}, { Definition->Features.SDKs[SdkIndex].Name }, {}, {});
 			}
 		}
@@ -367,47 +387,49 @@ void FRigMapperDefinitionEditorToolkit::HandleStructureSelectionChanged(ESelectI
 
 void FRigMapperDefinitionEditorToolkit::SetFeatureDetailsView(const URigMapperDefinitionEditorGraphNode* InNode)
 {
-	URigMapperDefinitionEditorGraphNode* Node = ConstCast<const URigMapperDefinitionEditorGraphNode>(InNode);
+	TObjectPtr<URigMapperDefinitionEditorGraphNode> Node = ConstCast<const URigMapperDefinitionEditorGraphNode>(InNode);
 	if (!Node || !Node->GetNodeFeature())
 	{
 		return;
 	}
 
-	if (Node->GetNodeType() == ERigMapperNodeType::SDK)
+	if (Node->GetNodeType() == ERigMapperFeatureType::SDK)
 	{
 		const TSharedPtr<FStructOnScope> SdkStructure = MakeShared<FStructOnScope>(FRigMapperSdkFeature::StaticStruct(), reinterpret_cast<uint8*>(Node->GetNodeFeature()));
 		FeatureDetailsView->SetStructureData(SdkStructure);
 	}
-	else if (Node->GetNodeType() == ERigMapperNodeType::WeightedSum)
+	else if (Node->GetNodeType() == ERigMapperFeatureType::WeightedSum)
 	{
 		const TSharedPtr<FStructOnScope> WsStructure = MakeShared<FStructOnScope>(FRigMapperWsFeature::StaticStruct(), reinterpret_cast<uint8*>(Node->GetNodeFeature()));
 		FeatureDetailsView->SetStructureData(WsStructure);
 	}
-	else if (Node->GetNodeType() == ERigMapperNodeType::Multiply)
+	else if (Node->GetNodeType() == ERigMapperFeatureType::Multiply)
 	{
 		const TSharedPtr<FStructOnScope> MultiplyStructure = MakeShared<FStructOnScope>(FRigMapperMultiplyFeature::StaticStruct(), reinterpret_cast<uint8*>(Node->GetNodeFeature()));
 		FeatureDetailsView->SetStructureData(MultiplyStructure);
 	}
-	else if (Node->GetNodeType() == ERigMapperNodeType::Output)
+	else if (Node->GetNodeType() == ERigMapperFeatureType::Output)
 	{
 		const TSharedPtr<FStructOnScope> OutputStructure = MakeShared<FStructOnScope>(FRigMapperOutputFeature::StaticStruct(), reinterpret_cast<uint8*>(Node->GetNodeFeature()));
 		FeatureDetailsView->SetStructureData(OutputStructure);
 	}
-	else if (Node->GetNodeType() == ERigMapperNodeType::Input || Node->GetNodeType() == ERigMapperNodeType::NullOutput)
+	else if (Node->GetNodeType() == ERigMapperFeatureType::Input || Node->GetNodeType() == ERigMapperFeatureType::NullOutput)
 	{
 		const TSharedPtr<FStructOnScope> FeatureStructure = MakeShared<FStructOnScope>(FRigMapperFeature::StaticStruct(), reinterpret_cast<uint8*>(Node->GetNodeFeature()));
 		FeatureDetailsView->SetStructureData(FeatureStructure);
-	}	
+	}			
 	FeatureDetailsView->GetOnFinishedChangingPropertiesDelegate().Clear();
-	FeatureDetailsView->GetOnFinishedChangingPropertiesDelegate().AddLambda(
-		[WeakNode = Node]
-		(const FPropertyChangedEvent&)
-		{
-			if (WeakNode)
-			{
-				WeakNode->UpdateNodeFromFeature();
-			}
-		});
+	FeatureDetailsView->GetOnFinishedChangingPropertiesDelegate().AddRaw(this, &FRigMapperDefinitionEditorToolkit::HandleFinishedChangingFeatureProperties, TWeakObjectPtr<URigMapperDefinitionEditorGraphNode>(Node.Get()));	
+}
+
+void FRigMapperDefinitionEditorToolkit::HandleFinishedChangingFeatureProperties(const FPropertyChangedEvent& InPropertyChangedEvent, TWeakObjectPtr<URigMapperDefinitionEditorGraphNode> InNode)
+{
+	if (InNode.IsValid())
+	{
+		InNode->UpdateNodeFromFeature(InPropertyChangedEvent);
+		GraphEditor->RefreshGraphNode(InNode.Get());
+	}	
+
 }
 
 #undef LOCTEXT_NAMESPACE
\ No newline at end of file
diff --git a/Engine/Plugins/Experimental/Animation/RigMapper/Source/RigMapperEditor/Private/RigMapperDefinitionEditorToolkit.h b/Engine/Plugins/Experimental/Animation/RigMapper/Source/RigMapperEditor/Private/RigMapperDefinitionEditorToolkit.h
--- a/Engine/Plugins/Experimental/Animation/RigMapper/Source/RigMapperEditor/Private/RigMapperDefinitionEditorToolkit.h
+++ b/Engine/Plugins/Experimental/Animation/RigMapper/Source/RigMapperEditor/Private/RigMapperDefinitionEditorToolkit.h
@@ -10,7 +10,7 @@
 
 #define UE_API RIGMAPPEREDITOR_API
 
-enum class ERigMapperNodeType : uint8;
+enum class ERigMapperFeatureType : uint8;
 class URigMapperDefinitionEditorGraphNode;
 class SRigMapperDefinitionStructureView;
 class SRigMapperDefinitionGraphEditor;
@@ -54,6 +54,7 @@ private:
 
 	UE_API bool HandleIsPropertyVisible(const FPropertyAndParent& PropertyAndParent);
 	UE_API void HandleFinishedChangingProperties(const FPropertyChangedEvent& PropertyChangedEvent);
+	UE_API void HandleFinishedChangingFeatureProperties(const FPropertyChangedEvent& PropertyChangedEvent, TWeakObjectPtr<URigMapperDefinitionEditorGraphNode> InNode);
 	UE_API void HandleRigMapperDefinitionLoaded();
 	UE_API void HandleGraphSelectionChanged(const TSet<UObject*>& Nodes);
 	
@@ -65,7 +66,7 @@ private:
 	static UE_API const FName DefinitionEditorStructureTabId;
 	static UE_API const FName DefinitionEditorDetailsTabId;
 
-	static UE_API const TMap<FName, ERigMapperNodeType> PropertyNameToNodeTypeMapping;
+	static UE_API const TMap<FName, ERigMapperFeatureType> PropertyNameToNodeTypeMapping;
 	
 	TObjectPtr<URigMapperDefinition> Definition;
 
diff --git a/Engine/Plugins/Experimental/Animation/RigMapper/Source/RigMapperEditor/Private/RigMapperGraph/RigMapperDefinitionEditorGraph.cpp b/Engine/Plugins/Experimental/Animation/RigMapper/Source/RigMapperEditor/Private/RigMapperGraph/RigMapperDefinitionEditorGraph.cpp
--- a/Engine/Plugins/Experimental/Animation/RigMapper/Source/RigMapperEditor/Private/RigMapperGraph/RigMapperDefinitionEditorGraph.cpp
+++ b/Engine/Plugins/Experimental/Animation/RigMapper/Source/RigMapperEditor/Private/RigMapperGraph/RigMapperDefinitionEditorGraph.cpp
@@ -1,7 +1,6 @@
 // Copyright Epic Games, Inc. All Rights Reserved.
 
 #include "RigMapperDefinitionEditorGraph.h"
-
 #include "RigMapperDefinition.h"
 #include "RigMapperDefinitionEditorGraphNode.h"
 #include "RigMapperDefinitionEditorGraphSchema.h"
@@ -23,7 +22,7 @@ void URigMapperDefinitionEditorGraph::RebuildGraph()
 	RemoveAllNodes();
 	ConstructNodes();
 	NotifyGraphChanged();
-	RequestRefreshLayout(false);
+	RequestRefreshLayout(true);
 }
 
 void URigMapperDefinitionEditorGraph::ConstructNodes()
@@ -53,14 +52,14 @@ void URigMapperDefinitionEditorGraph::ConstructNodes()
 	{
 		if (!InputNodes.Contains(Input))
 		{
-			CreateGraphNode(Input, ERigMapperNodeType::Input);
+			CreateGraphNode(Input, ERigMapperFeatureType::Input);
 		}
 	}
 	for (const FString& NullOutput : Definition->NullOutputs)
 	{
 		if (!NullOutputNodes.Contains(NullOutput))
 		{
-			CreateGraphNode(NullOutput, ERigMapperNodeType::NullOutput);
+			CreateGraphNode(NullOutput, ERigMapperFeatureType::NullOutput);
 		}
 	}
 }
@@ -89,6 +88,97 @@ TArray<URigMapperDefinitionEditorGraphNode*> URigMapperDefinitionEditorGraph::Ge
 	return OutNodes;
 }
 
+// TODO: Add RigMapperDefinition node subclasses
+URigMapperDefinitionEditorGraphNode* URigMapperDefinitionEditorGraph::FindNodeByName(const FString& InNodeName)
+{
+	if (URigMapperDefinitionEditorGraphNode* FoundNode = FeatureNodes.FindRef(InNodeName))
+	{
+		return FoundNode;
+	}
+	if (URigMapperDefinitionEditorGraphNode* FoundNode = InputNodes.FindRef(InNodeName))
+	{
+		return FoundNode;
+	}
+	if (URigMapperDefinitionEditorGraphNode* FoundNode = OutputNodes.FindRef(InNodeName))
+	{
+		return FoundNode;
+	}
+	if (URigMapperDefinitionEditorGraphNode* FoundNode = NullOutputNodes.FindRef(InNodeName))
+	{
+		return FoundNode;
+	}
+	return nullptr;
+}
+
+// TODO: Add RigMapperDefinition node subclasses
+bool URigMapperDefinitionEditorGraph::UpdateNodeNameInMap(const FString& InOldName, const FString& InNewName, const ERigMapperFeatureType InNodeType)
+{
+	switch (InNodeType)
+	{
+	case ERigMapperFeatureType::WeightedSum:
+	case ERigMapperFeatureType::SDK:
+	case ERigMapperFeatureType::Multiply:
+		if (URigMapperDefinitionEditorGraphNode* FoundNode = FeatureNodes.FindRef(InOldName))
+		{
+			FeatureNodes.Add(InNewName, FoundNode);
+			FeatureNodes.Remove(InOldName);
+			return true;
+		}
+		break;
+	case ERigMapperFeatureType::Input:
+		if (URigMapperDefinitionEditorGraphNode* FoundNode = InputNodes.FindRef(InOldName))
+		{
+			InputNodes.Add(InNewName, FoundNode);
+			InputNodes.Remove(InOldName);
+			return true;
+		}
+		break;
+	case ERigMapperFeatureType::Output:
+		if (URigMapperDefinitionEditorGraphNode* FoundNode = OutputNodes.FindRef(InOldName))
+		{
+			OutputNodes.Add(InNewName, FoundNode);
+			OutputNodes.Remove(InOldName);
+			return true;
+		}
+		break;
+	case ERigMapperFeatureType::NullOutput:
+		if (URigMapperDefinitionEditorGraphNode* FoundNode = NullOutputNodes.FindRef(InOldName))
+		{
+			NullOutputNodes.Add(InNewName, FoundNode);
+			NullOutputNodes.Remove(InOldName);
+			return true;
+		}
+		break;
+	default:
+		break;
+	}
+	return false;
+}
+
+bool URigMapperDefinitionEditorGraph::RemoveNodeFromMap(const FString& InNodeName, const ERigMapperFeatureType InNodeType)
+{
+	switch (InNodeType)
+	{
+	case ERigMapperFeatureType::WeightedSum:
+	case ERigMapperFeatureType::SDK:
+	case ERigMapperFeatureType::Multiply:
+		FeatureNodes.Remove(InNodeName);
+		break;
+	case ERigMapperFeatureType::Input:
+		InputNodes.Remove(InNodeName);
+		break;
+	case ERigMapperFeatureType::Output:
+		OutputNodes.Remove(InNodeName);
+		break;
+	case ERigMapperFeatureType::NullOutput:
+		NullOutputNodes.Remove(InNodeName);
+		break;
+	default:
+		return false;
+	}
+	return true;
+}
+
 URigMapperDefinitionEditorGraphNode* URigMapperDefinitionEditorGraph::CreateGraphNodesRec(URigMapperDefinition* Definition, const FString& NodeName, bool bIsOutputNode)
 {
 	// No need to do anything is node was already created. Output nodes should not be referenced by other nodes and should only get created once.
@@ -112,7 +202,7 @@ URigMapperDefinitionEditorGraphNode* URigMapperDefinitionEditorGraph::CreateGrap
 	}
 	if (Definition->Inputs.Contains(NodeName))
 	{
-		return CreateGraphNode(NodeName, ERigMapperNodeType::Input); 
+		return CreateGraphNode(NodeName, ERigMapperFeatureType::Input);
 	}
 	return CreateFeatureNode(Definition, NodeName);
 }
@@ -121,7 +211,7 @@ URigMapperDefinitionEditorGraphNode* URigMapperDefinitionEditorGraph::CreateOutp
 {
 	if (const FString* LinkedInputName = Definition->Outputs.Find(NodeName))
 	{
-		URigMapperDefinitionEditorGraphNode* Node = CreateGraphNode(NodeName, ERigMapperNodeType::Output);
+		URigMapperDefinitionEditorGraphNode* Node = CreateGraphNode(NodeName, ERigMapperFeatureType::Output);
 		
 		if (URigMapperDefinitionEditorGraphNode* LinkedNode = CreateGraphNodesRec(Definition, *LinkedInputName, false))
 		{
@@ -139,8 +229,8 @@ URigMapperDefinitionEditorGraphNode* URigMapperDefinitionEditorGraph::CreateFeat
 
 	if (const FRigMapperFeature* Feature = Definition->Features.Find(NodeName, FeatureType))
 	{
-		const ERigMapperNodeType NodeType = static_cast<ERigMapperNodeType>(static_cast<uint8>(FeatureType));
-		URigMapperDefinitionEditorGraphNode* Node = CreateGraphNode(NodeName, NodeType);
+		
+		URigMapperDefinitionEditorGraphNode* Node = CreateGraphNode(NodeName, FeatureType);
 			
 		TArray<FString> FeatureInputNames;
 		Feature->GetInputs(FeatureInputNames);
@@ -179,7 +269,7 @@ void URigMapperDefinitionEditorGraph::LinkGraphNodes(URigMapperDefinitionEditorG
 }
 		
 
-URigMapperDefinitionEditorGraphNode* URigMapperDefinitionEditorGraph::CreateGraphNode(const FString& NodeName, ERigMapperNodeType NodeType)
+URigMapperDefinitionEditorGraphNode* URigMapperDefinitionEditorGraph::CreateGraphNode(const FString& NodeName, const ERigMapperFeatureType NodeType)
 {
 	if (!WeakDefinition.IsValid())
 	{
@@ -194,28 +284,13 @@ URigMapperDefinitionEditorGraphNode* URigMapperDefinitionEditorGraph::CreateGrap
 
 	Node->SetupNode(WeakDefinition.Get(), NodeName, NodeType);
 
-	if (NodeType == ERigMapperNodeType::Input)
-	{
-		InputNodes.Add(NodeName, Node);
-	}
-	else if (NodeType == ERigMapperNodeType::Output)
-	{
-		OutputNodes.Add(NodeName, Node);
-	}
-	else if (NodeType == ERigMapperNodeType::NullOutput)
-	{
-		NullOutputNodes.Add(NodeName, Node);
-	}
-	else // todo: Add feature type
-	{
-		FeatureNodes.Add(NodeName, Node);
-	}
+	AddNodeToMap(NodeType, NodeName, Node);
 	
 	return Node;
 }
 
 
-URigMapperDefinitionEditorGraphNode* URigMapperDefinitionEditorGraph::CreateNewGraphNode(const ERigMapperNodeType InNodeType, UEdGraphPin* InFromPin, const FVector2f& InLocation, bool bSelectNewNode /*= true*/)
+URigMapperDefinitionEditorGraphNode* URigMapperDefinitionEditorGraph::CreateNewGraphNode(const ERigMapperFeatureType InNodeType, UEdGraphPin* InFromPin, const FVector2f& InLocation, bool bSelectNewNode /*= true*/)
 {
 	if (!WeakDefinition.IsValid())
 	{
@@ -226,11 +301,13 @@ URigMapperDefinitionEditorGraphNode* URigMapperDefinitionEditorGraph::CreateNewG
 	{
 		FromPinName = Cast<URigMapperDefinitionEditorGraphNode>(InFromPin->GetOwningNode())->GetNodeName();
 	}
-	URigMapperDefinitionEditorGraphNode* NewNode = CreateGraphNode(MakeUniqueObjectName(this, URigMapperDefinitionEditorGraphNode::StaticClass(), FName(FromPinName)).ToString(), InNodeType);
+	FGraphNodeCreator<URigMapperDefinitionEditorGraphNode> GraphNodeCreator(*this);
+	URigMapperDefinitionEditorGraphNode* NewNode = GraphNodeCreator.CreateNode(bSelectNewNode);
+	GraphNodeCreator.Finalize();
 	NewNode->NodePosX = InLocation.X;
 	NewNode->NodePosY = InLocation.Y;
-	FString NodeName = NewNode->GetNodeName();
-	if (InNodeType != ERigMapperNodeType::Input && InNodeType != ERigMapperNodeType::NullOutput && InNodeType != ERigMapperNodeType::Invalid)
+	FString NodeName = MakeUniqueObjectName(this, URigMapperDefinitionEditorGraphNode::StaticClass(), FName(FromPinName)).ToString();
+	if (InNodeType != ERigMapperFeatureType::Input && InNodeType != ERigMapperFeatureType::NullOutput)
 	{
 		UEdGraphPin* InPin = NewNode->CreateInputPin();
 		InPin->bHidden = false;
@@ -240,7 +317,7 @@ URigMapperDefinitionEditorGraphNode* URigMapperDefinitionEditorGraph::CreateNewG
 	{
 		InputNodeName = FromPinName;
 	}
-	if (InNodeType == ERigMapperNodeType::Multiply)
+	if (InNodeType == ERigMapperFeatureType::Multiply)
 	{
 		FRigMapperMultiplyFeature NewFeature = FRigMapperMultiplyFeature(NodeName);
 		if (!InputNodeName.IsEmpty())
@@ -252,7 +329,7 @@ URigMapperDefinitionEditorGraphNode* URigMapperDefinitionEditorGraph::CreateNewG
 		UEdGraphPin* InPin = NewNode->CreateInputPin();
 		InPin->bHidden = false;
 	}
-	else if (InNodeType == ERigMapperNodeType::WeightedSum)
+	else if (InNodeType == ERigMapperFeatureType::WeightedSum)
 	{
 		FRigMapperWsFeature NewFeature = FRigMapperWsFeature(NodeName);
 		if (!InputNodeName.IsEmpty())
@@ -261,7 +338,7 @@ URigMapperDefinitionEditorGraphNode* URigMapperDefinitionEditorGraph::CreateNewG
 		}
 		WeakDefinition->Features.WeightedSums.Add(NewFeature);
 	}
-	else if (InNodeType == ERigMapperNodeType::SDK)
+	else if (InNodeType == ERigMapperFeatureType::SDK)
 	{
 		FRigMapperSdkFeature NewFeature = FRigMapperSdkFeature(NodeName);
 		if (!InputNodeName.IsEmpty())
@@ -270,19 +347,19 @@ URigMapperDefinitionEditorGraphNode* URigMapperDefinitionEditorGraph::CreateNewG
 		}
 		WeakDefinition->Features.SDKs.Add(NewFeature);
 	}
-	else if (InNodeType == ERigMapperNodeType::Output)
+	else if (InNodeType == ERigMapperFeatureType::Output)
 	{
 		WeakDefinition->Outputs.Add(NodeName, InputNodeName);
 	}
-	else if (InNodeType == ERigMapperNodeType::Input)
+	else if (InNodeType == ERigMapperFeatureType::Input)
 	{
 		WeakDefinition->Inputs.Add(NodeName);
 	}
-	else if (InNodeType == ERigMapperNodeType::NullOutput)
+	else if (InNodeType == ERigMapperFeatureType::NullOutput)
 	{
 		WeakDefinition->NullOutputs.Add(NodeName);
 	}
-	if (InNodeType != ERigMapperNodeType::Output && InNodeType != ERigMapperNodeType::NullOutput)
+	if (InNodeType != ERigMapperFeatureType::Output && InNodeType != ERigMapperFeatureType::NullOutput)
 	{
 		UEdGraphPin* OutPin = NewNode->CreateOutputPin();
 		OutPin->bHidden = false;
@@ -290,6 +367,8 @@ URigMapperDefinitionEditorGraphNode* URigMapperDefinitionEditorGraph::CreateNewG
 
 	NewNode->SetupNode(WeakDefinition.Get(), NodeName, InNodeType);
 
+	AddNodeToMap(InNodeType, NodeName, NewNode);
+
 	if (InFromPin)
 	{
 		const UEdGraphSchema* GraphSchema = GetSchema();
@@ -354,16 +433,16 @@ void URigMapperDefinitionEditorGraph::LayoutNodeRec(URigMapperDefinitionEditorGr
 	const FVector2D& Dimensions = InNode->GetDimensions();
 	FVector2D Margin = { NodeMarginX, NodeMarginY };
 	
-	if (InNode->GetNodeType() == ERigMapperNodeType::Input)
+	if (InNode->GetNodeType() == ERigMapperFeatureType::Input)
 	{
 		InNode->NodePosX = 0;
 		Margin.X += InputsWidth - Dimensions.X;
 	}
-	else if (InNode->GetNodeType() == ERigMapperNodeType::Output)
+	else if (InNode->GetNodeType() == ERigMapperFeatureType::Output)
 	{
 		Margin.X = 0;
 	}
-	else if (InNode->GetNodeType() == ERigMapperNodeType::NullOutput)
+	else if (InNode->GetNodeType() == ERigMapperFeatureType::NullOutput)
 	{
 		Margin.X = 0;
 	}
@@ -382,6 +461,27 @@ void URigMapperDefinitionEditorGraph::LayoutNodeRec(URigMapperDefinitionEditorGr
 	InNode->SetMargin(Margin);
 }
 
+// TODO: Add RigMapperDefinition node subclasses
+void URigMapperDefinitionEditorGraph::AddNodeToMap(const ERigMapperFeatureType NodeType, const FString& NodeName, URigMapperDefinitionEditorGraphNode* Node)
+{
+	if (NodeType == ERigMapperFeatureType::Input)
+	{
+		InputNodes.Add(NodeName, Node);
+	}
+	else if (NodeType == ERigMapperFeatureType::Output)
+	{
+		OutputNodes.Add(NodeName, Node);
+	}
+	else if (NodeType == ERigMapperFeatureType::NullOutput)
+	{
+		NullOutputNodes.Add(NodeName, Node);
+	}
+	else // todo: Add feature type
+	{
+		FeatureNodes.Add(NodeName, Node);
+	}
+}
+
 void URigMapperDefinitionEditorGraph::LayoutNodes() const
 {
 	double InputsMaxWidth = 0;
diff --git a/Engine/Plugins/Experimental/Animation/RigMapper/Source/RigMapperEditor/Private/RigMapperGraph/RigMapperDefinitionEditorGraph.h b/Engine/Plugins/Experimental/Animation/RigMapper/Source/RigMapperEditor/Private/RigMapperGraph/RigMapperDefinitionEditorGraph.h
--- a/Engine/Plugins/Experimental/Animation/RigMapper/Source/RigMapperEditor/Private/RigMapperGraph/RigMapperDefinitionEditorGraph.h
+++ b/Engine/Plugins/Experimental/Animation/RigMapper/Source/RigMapperEditor/Private/RigMapperGraph/RigMapperDefinitionEditorGraph.h
@@ -8,7 +8,7 @@
 
 #define UE_API RIGMAPPEREDITOR_API
 
-enum class ERigMapperNodeType : uint8;
+enum class ERigMapperFeatureType : uint8;
 class URigMapperDefinitionEditorGraphNode;
 class URigMapperDefinition;
 
@@ -32,9 +32,15 @@ public:
 	
 	URigMapperDefinition* GetDefinition() const { return WeakDefinition.Get(); };
 
-	UE_API URigMapperDefinitionEditorGraphNode* CreateGraphNode(const FString& NodeName, ERigMapperNodeType NodeType);
-	UE_API URigMapperDefinitionEditorGraphNode* CreateNewGraphNode(const ERigMapperNodeType InNodeType, UEdGraphPin* InFromPin, const FVector2f& InLocation, bool bSelectNewNode = true);
+	UE_API URigMapperDefinitionEditorGraphNode* CreateGraphNode(const FString& NodeName, const ERigMapperFeatureType NodeType);
+	UE_API URigMapperDefinitionEditorGraphNode* CreateNewGraphNode(const ERigMapperFeatureType InNodeType, UEdGraphPin* InFromPin, const FVector2f& InLocation, bool bSelectNewNode = true);
 	UE_API TArray<URigMapperDefinitionEditorGraphNode*> GetNodesByName(const TArray<FString>& Inputs, const TArray<FString>& Features, const TArray<FString>& Outputs, const TArray<FString>& NullOutputs) const;
+	/** Retrieves first node that matches the InNodeName (assumes node names are unique	*/
+	UE_API URigMapperDefinitionEditorGraphNode* FindNodeByName(const FString& InNodeName);
+	/** Updates the key in the node map when the node is renamed (assumes new name is unique) */
+	UE_API bool UpdateNodeNameInMap(const FString& InOldName, const FString& InNewName, const ERigMapperFeatureType InNodeType);
+	/** Removes node from the map when deleted */
+	UE_API bool RemoveNodeFromMap(const FString& InNodeName, const ERigMapperFeatureType InNodeType);
 
 private:
 	UE_API URigMapperDefinitionEditorGraphNode* CreateGraphNodesRec(URigMapperDefinition* Definition, const FString& InNodeName, bool bIsOutputNode);
@@ -43,6 +49,9 @@ private:
 	static UE_API void LinkGraphNodes(URigMapperDefinitionEditorGraphNode* InNode, URigMapperDefinitionEditorGraphNode* OutNode);
 	UE_API void LayoutNodeRec(URigMapperDefinitionEditorGraphNode* InNode, double InputsWidth, double PosY, TArray<URigMapperDefinitionEditorGraphNode*>& LayedOutNodes) const;
 	
+
+	void AddNodeToMap(const ERigMapperFeatureType NodeType, const FString& NodeName, URigMapperDefinitionEditorGraphNode* Node);
+
 private:
 	TWeakObjectPtr<URigMapperDefinition> WeakDefinition;
 	TMap<FString, URigMapperDefinitionEditorGraphNode*> InputNodes;
diff --git a/Engine/Plugins/Experimental/Animation/RigMapper/Source/RigMapperEditor/Private/RigMapperGraph/RigMapperDefinitionEditorGraphNode.cpp b/Engine/Plugins/Experimental/Animation/RigMapper/Source/RigMapperEditor/Private/RigMapperGraph/RigMapperDefinitionEditorGraphNode.cpp
--- a/Engine/Plugins/Experimental/Animation/RigMapper/Source/RigMapperEditor/Private/RigMapperGraph/RigMapperDefinitionEditorGraphNode.cpp
+++ b/Engine/Plugins/Experimental/Animation/RigMapper/Source/RigMapperEditor/Private/RigMapperGraph/RigMapperDefinitionEditorGraphNode.cpp
@@ -6,6 +6,7 @@
 #include "SRigMapperDefinitionGraphEditorNode.h"
 #include "RigMapperDefinition.h"
 #include "RigMapperDefinitionEditorGraph.h"
+#include "PropertyEditorModule.h"
 #include "EdGraph/EdGraphPin.h"
 
 #include UE_INLINE_GENERATED_CPP_BY_NAME(RigMapperDefinitionEditorGraphNode)
@@ -23,7 +24,7 @@ TSharedPtr<SGraphNode> URigMapperDefinitionEditorGraphNode::NodeFactory::CreateN
 
 		TSharedRef<SGraphNode> GraphNode = SNew(SRigMapperDefinitionGraphEditorNode, ThisNode);
 		GraphNode->SlatePrepass();
-		ThisNode->SetDimensions(GraphNode->GetDesiredSize());
+		ThisNode->SetDimensions(GraphNode->GetDesiredSize());		
 		return GraphNode;
 	}
 	return nullptr;
@@ -36,27 +37,27 @@ FLinearColor URigMapperDefinitionEditorGraphNode::GetNodeBodyTintColor() const
 
 FLinearColor URigMapperDefinitionEditorGraphNode::GetNodeTitleColor() const
 {
-	if (NodeType == ERigMapperNodeType::Input)
+	if (NodeType == ERigMapperFeatureType::Input)
 	{
 		return FLinearColor(0.466, 0.969, 0.878);
 	}
-	if (NodeType == ERigMapperNodeType::Output)
+	if (NodeType == ERigMapperFeatureType::Output)
 	{
 		return FLinearColor(0.466, 0.969, 0.525);
 	}
-	if (NodeType == ERigMapperNodeType::NullOutput)
+	if (NodeType == ERigMapperFeatureType::NullOutput)
 	{
 		return FLinearColor(0.65, 0.969, 0.466);
 	}
-	if (NodeType == ERigMapperNodeType::Multiply)
+	if (NodeType == ERigMapperFeatureType::Multiply)
 	{
 		return FLinearColor(0.969, 0.966, 0.466);
 	}
-	if (NodeType == ERigMapperNodeType::SDK)
+	if (NodeType == ERigMapperFeatureType::SDK)
 	{
 		return FLinearColor(0.969, 0.717, 0.466);
 	}
-	if (NodeType == ERigMapperNodeType::WeightedSum)
+	if (NodeType == ERigMapperFeatureType::WeightedSum)
 	{
 		return FLinearColor(0.65, 0.466, 0.969);
 	}
@@ -73,45 +74,33 @@ void URigMapperDefinitionEditorGraphNode::AllocateDefaultPins()
 
 }
 
-void URigMapperDefinitionEditorGraphNode::SetupNode(URigMapperDefinition* InDefinition, const FString& InNodeName, ERigMapperNodeType InNodeType)
+void URigMapperDefinitionEditorGraphNode::SetupNode(URigMapperDefinition* InDefinition, const FString& InNodeName, ERigMapperFeatureType InNodeType)
 {
 	NodeName = InNodeName;
 	Definition = InDefinition;
 	NodeType = InNodeType;
-	NodeTitle = FText::FromString(NodeName);
-
+	NodeTitle = FText::FromString(NodeName);	
 	if (!Definition.IsValid())
 	{
 		return;
 	}
-	FeaturePtr = nullptr;
-	IsExternalOwner = false;
 
-	if (NodeType == ERigMapperNodeType::Input && Definition.Get()->Inputs.Contains(NodeName))
-	{
-		FeaturePtr = new FRigMapperFeature(NodeName);		
-	}
-	else if (NodeType == ERigMapperNodeType::WeightedSum || NodeType == ERigMapperNodeType::SDK || NodeType == ERigMapperNodeType::Multiply)
+	// Create features not tracked with Definition and assign the Definition as the owner.
+	if (NodeType == ERigMapperFeatureType::Output)
 	{
-		IsExternalOwner = true;
-		ERigMapperFeatureType FeatureType;
-		if (FRigMapperFeature* Feature = Definition.Get()->Features.Find(NodeName, FeatureType))
-		{
-			FeaturePtr = Feature;
-		}
+		FRigMapperOutputFeature NewFeature = FRigMapperOutputFeature(NodeName);
+		NewFeature.Input = Definition->Outputs[NodeName];
+		Definition->Features.Outputs.Add(NewFeature);
 	}
-	else if (NodeType == ERigMapperNodeType::Output)
+	else if (NodeType == ERigMapperFeatureType::Input)
 	{
-		if (FString* LinkedInputName = Definition.Get()->Outputs.Find(NodeName))
-		{
-			FRigMapperOutputFeature* OutputFeature = new FRigMapperOutputFeature(NodeName);
-			OutputFeature->Input = *LinkedInputName;			
-			FeaturePtr = static_cast<FRigMapperFeature*>(OutputFeature);
-		}
+		FRigMapperFeature NewFeature = FRigMapperFeature(NodeName);
+		Definition->Features.Inputs.Add(NewFeature);
 	}
-	else if (NodeType == ERigMapperNodeType::NullOutput && Definition.Get()->NullOutputs.Contains(NodeName))
+	else if (NodeType == ERigMapperFeatureType::NullOutput)
 	{
-		FeaturePtr = new FRigMapperNullOutputFeature(NodeName);
+		FRigMapperNullOutputFeature NewFeature = FRigMapperNullOutputFeature(NodeName);
+		Definition->Features.NullOutputs.Add(NewFeature);
 	}
 
 	UpdateNodeFields();
@@ -119,18 +108,18 @@ void URigMapperDefinitionEditorGraphNode::SetupNode(URigMapperDefinition* InDefi
 
 void URigMapperDefinitionEditorGraphNode::UpdateNodeFields()
 {
-
 	if (!Definition.IsValid())
 	{
 		return;
 	}
 
+	FRigMapperFeature* FeaturePtr = GetNodeFeature();
 	if (!FeaturePtr)
 	{
 		return;
 	}
 
-	if (NodeType == ERigMapperNodeType::WeightedSum || NodeType == ERigMapperNodeType::SDK)
+	if (NodeType == ERigMapperFeatureType::WeightedSum || NodeType == ERigMapperFeatureType::SDK)
 	{
 		if (FeaturePtr->GetFeatureType() == ERigMapperFeatureType::SDK)
 		{
@@ -190,123 +179,57 @@ void URigMapperDefinitionEditorGraphNode::UpdateNodeFields()
 
 void URigMapperDefinitionEditorGraphNode::RelinkFeature(const FString& InOldInput, const FString& InNewInput)
 {
-	if (!Definition.IsValid() || !FeaturePtr)
+	if (!Definition.IsValid())
+	{
+		return;
+	}
+	FRigMapperFeature* FeaturePtr = GetNodeFeature();
+	if (!FeaturePtr)
 	{
 		return;
 	}
 
-	if (NodeType == ERigMapperNodeType::Output)
+	FeaturePtr->RelinkFeature(InOldInput, InNewInput);
+	if (NodeType == ERigMapperFeatureType::Output)
 	{
-		FRigMapperOutputFeature* Feature = static_cast<FRigMapperOutputFeature*>(FeaturePtr);
-		if (Feature)
-		{
-			Feature->Input = InNewInput;
-		}
+		// Definition for Outputs has to be updated manually since RigMapperOutputFeature is created separately.
 		if (Definition->Outputs.Contains(NodeName))
 		{
 			Definition->Outputs[NodeName] = InNewInput;
 		}
-	}
-	else if (NodeType == ERigMapperNodeType::Multiply)
-	{
-		FRigMapperMultiplyFeature* Feature = static_cast<FRigMapperMultiplyFeature*>(FeaturePtr);
-		if (Feature)
-		{
-			for (int32 i = 0; i < Feature->Inputs.Num(); i++)
-			{
-				if (Feature->Inputs[i] == InOldInput)
-				{
-					Feature->Inputs[i] = InNewInput;
-				}
-			}
-		}
-	}
-	else if (NodeType == ERigMapperNodeType::WeightedSum)
-	{
-		FRigMapperWsFeature* Feature = static_cast<FRigMapperWsFeature*>(FeaturePtr);
-		if (Feature)
-		{
-			if (double* Weight = Feature->Inputs.Find(InOldInput))
-			{
-				Feature->Inputs.Remove(InOldInput);
-				Feature->Inputs.Add(InNewInput, MoveTemp(*Weight));
-			}
-		}
-	}
-	else if (NodeType == ERigMapperNodeType::SDK)
-	{
-		FRigMapperSdkFeature* Feature = static_cast<FRigMapperSdkFeature*>(FeaturePtr);
-		if (Feature)
-		{
-			Feature->Input = InNewInput;
-		}
 	}	
 }
 
 void URigMapperDefinitionEditorGraphNode::AddNewInputLink(const FString& InNewInput)
 {
-	if (!Definition.IsValid() || !FeaturePtr)
+	if (!Definition.IsValid())
 	{
 		return;
 	}
 
-	if (NodeType == ERigMapperNodeType::Output)
-	{
-		FRigMapperOutputFeature* Feature = static_cast<FRigMapperOutputFeature*>(FeaturePtr);
-		if (Feature)
-		{
-			Feature->Input = InNewInput;
-		}
-		if (Definition->Outputs.Contains(NodeName))
-		{
-			Definition->Outputs[NodeName] = InNewInput;
-		}
-	}
-	else if (NodeType == ERigMapperNodeType::Multiply)
-	{
-		FRigMapperMultiplyFeature* Feature = static_cast<FRigMapperMultiplyFeature*>(FeaturePtr);
-		if (Feature)
-		{
-			Feature->Inputs.Add(InNewInput);
-		}
-	}
-	else if (NodeType == ERigMapperNodeType::WeightedSum)
+	FRigMapperFeature* FeaturePtr = GetNodeFeature();
+	if (!FeaturePtr)
 	{
-		FRigMapperWsFeature* Feature = static_cast<FRigMapperWsFeature*>(FeaturePtr);
-		if (Feature)
-		{
-			if (!Feature->Inputs.Contains(InNewInput))
-			{				
-				Feature->Inputs.Add(InNewInput, 1.0f);
-			}
-		}
+		return;
 	}
-	else if (NodeType == ERigMapperNodeType::SDK)
+	FeaturePtr->AddNewInputLink(InNewInput);
+	if (NodeType == ERigMapperFeatureType::Output)
 	{
-		FRigMapperSdkFeature* Feature = static_cast<FRigMapperSdkFeature*>(FeaturePtr);
-		if (Feature)
+		// Definition for Outputs has to be updated manually since RigMapperOutputFeature is created separately.
+		if (Definition->Outputs.Contains(NodeName))
 		{
-			Feature->Input = InNewInput;
+			Definition->Outputs[NodeName] = InNewInput;
 		}
 	}
 }
 
-void URigMapperDefinitionEditorGraphNode::BeginDestroy()
+FRigMapperFeature* URigMapperDefinitionEditorGraphNode::GetNodeFeature() const
 {
-	if (FeaturePtr)
+	if (Definition.IsValid())
 	{
-		if (!IsExternalOwner)
-		{
-			delete(FeaturePtr);
-		}
-		FeaturePtr = nullptr;
+		return Definition->GetFeature(NodeName, NodeType);
 	}
-	Super::BeginDestroy();
-}
-
-FRigMapperFeature* URigMapperDefinitionEditorGraphNode::GetNodeFeature() const
-{
-	return FeaturePtr;
+	return nullptr;
 }
 
 UEdGraphPin* URigMapperDefinitionEditorGraphNode::CreateInputPin()
@@ -319,6 +242,16 @@ UEdGraphPin* URigMapperDefinitionEditorGraphNode::CreateOutputPin()
 	return OutputPins.Add_GetRef(CreatePin(EEdGraphPinDirection::EGPD_Output, NAME_None, NAME_None, NAME_None));
 }
 
+bool URigMapperDefinitionEditorGraphNode::RemoveInputPin(UEdGraphPin* InPinToRemove)
+{	
+	bool bPinRemoved = RemovePin(InPinToRemove);
+	if (bPinRemoved)
+	{
+		InputPins.Remove(InPinToRemove);
+	}
+	return bPinRemoved;
+}
+
 void URigMapperDefinitionEditorGraphNode::GetRect(FVector2D& TopLeft, FVector2D& BottomRight) const
 {
 	TopLeft.X = NodePosX;
@@ -327,31 +260,144 @@ void URigMapperDefinitionEditorGraphNode::GetRect(FVector2D& TopLeft, FVector2D&
 	BottomRight.Y = NodePosY + GetDimensions().Y + GetMargin().Y;
 }
 
-void URigMapperDefinitionEditorGraphNode::UpdateNodeFromFeature()
+void URigMapperDefinitionEditorGraphNode::UpdateNodeFromFeature(const FPropertyChangedEvent& InPropertyChangedEvent)
 {
-	if (FeaturePtr)
+	FRigMapperFeature* FeaturePtr = GetNodeFeature();
+	if (!FeaturePtr)
 	{
-		FString NewName = FeaturePtr->Name;
-		
-		if (NewName != NodeName && NodeType != ERigMapperNodeType::Output && NodeType != ERigMapperNodeType::NullOutput)
+		return;
+	}
+
+	Modify();
+
+	const FName PropertyName = InPropertyChangedEvent.GetPropertyName();
+	const FName ParentName = InPropertyChangedEvent.GetMemberPropertyName();
+	// TODO: Refactor with specialized node subclasses for each type
+	if (NodeType == ERigMapperFeatureType::Multiply || NodeType == ERigMapperFeatureType::WeightedSum)
+	{
+		TMap<FString, int32> PropertyNameStack;
+		InPropertyChangedEvent.GetArrayIndicesPerObject(0, PropertyNameStack);
+		const FString InputsProperty = NodeType == ERigMapperFeatureType::WeightedSum ?
+			GET_MEMBER_NAME_STRING_CHECKED(FRigMapperWsFeature, OrderedInputs) :
+			GET_MEMBER_NAME_STRING_CHECKED(FRigMapperMultiplyFeature, Inputs);
+		const int32 PinIndex = PropertyNameStack.Contains(InputsProperty) ? PropertyNameStack.FindRef(InputsProperty) : INDEX_NONE;
+
+		const int32 MinReqPins = NodeType == ERigMapperFeatureType::WeightedSum ? 1 : 2;
+
+		if (InPropertyChangedEvent.ChangeType & EPropertyChangeType::ArrayClear)
+		{	
+			// REMOVING ALL INPUTS
+			// All inputs are removed but node has to have 2 input pin exposed.
+			for (int32 Index = InputPins.Num() - 1; Index >= MinReqPins; --Index)
+			{				
+				RemoveInputPin(InputPins[Index]);
+			}
+			// Break remaining input links.
+			for (UEdGraphPin* InPin : InputPins)
+			{
+				InPin->BreakAllPinLinks();
+			}
+			FeaturePtr->ClearInputs();
+		}
+		else if (InPropertyChangedEvent.ChangeType & EPropertyChangeType::ArrayRemove)
+		{
+			// REMOVING A SINGLE INPUT
+			if (PinIndex != INDEX_NONE && InputPins.IsValidIndex(PinIndex))
+			{
+				if (InputPins.Num() > MinReqPins)
+				{
+					RemoveInputPin(InputPins[PinIndex]);
+				}
+				else
+				{
+					InputPins[PinIndex]->BreakAllPinLinks();
+				}
+				FeaturePtr->UpdateFromDetails();
+			}
+		}
+		else if (InPropertyChangedEvent.ChangeType & EPropertyChangeType::ArrayAdd)
+		{
+			// ADDING A NEW (EMPTY) INPUT
+			if (NodeType == ERigMapperFeatureType::Multiply)
+			{
+				if (PinIndex != INDEX_NONE)
+				{
+					FRigMapperMultiplyFeature* Multiply = static_cast<FRigMapperMultiplyFeature*>(FeaturePtr);
+
+					if (Multiply->Inputs.Num() > InputPins.Num())
+					{
+						CreateInputPin();
+					}
+				}
+			}
+			else if (NodeType == ERigMapperFeatureType::WeightedSum)
+			{
+				if (PinIndex != INDEX_NONE)
+				{
+					FRigMapperWsFeature* WSFeature = static_cast<FRigMapperWsFeature*>(FeaturePtr);
+
+					if (WSFeature->OrderedInputs.Num() > InputPins.Num())
+					{
+						CreateInputPin();
+						WSFeature->UpdateFromDetails();
+					}
+				}
+			}
+		}
+		else if (InPropertyChangedEvent.ChangeType & EPropertyChangeType::ValueSet)
 		{
-			// Update inputs of all nodes directly connected to this node.
-			for (UEdGraphPin* OutPin : OutputPins) // It will always be one output pin.
+			// CHANGING AN INPUT NAME
+			if (PinIndex != INDEX_NONE)
 			{
-				for (UEdGraphPin* LinkedPin : OutPin->LinkedTo)
+				URigMapperDefinitionEditorGraph* RigMapperGraph = static_cast<URigMapperDefinitionEditorGraph*>(GetGraph());
+				URigMapperDefinitionEditorGraphNode* LinkedNode = nullptr;
+				if (NodeType == ERigMapperFeatureType::Multiply)
+				{
+					if (FRigMapperMultiplyFeature* Multiply = static_cast<FRigMapperMultiplyFeature*>(FeaturePtr))
+					{
+						const FString InputName = Multiply->Inputs[PinIndex];
+						LinkedNode = RigMapperGraph->FindNodeByName(InputName);
+					}
+				}
+				else if (NodeType == ERigMapperFeatureType::WeightedSum)
 				{
-					if (URigMapperDefinitionEditorGraphNode* LinkedNode = Cast<URigMapperDefinitionEditorGraphNode>(LinkedPin->GetOwningNode()))
+					if (FRigMapperWsFeature* WSFeature = static_cast<FRigMapperWsFeature*>(FeaturePtr))
 					{
-						LinkedNode->RelinkFeature(NodeName, NewName);
+						const FString InputName = WSFeature->OrderedInputs[PinIndex].Name;
+						LinkedNode = RigMapperGraph->FindNodeByName(InputName);
+						WSFeature->UpdateFromDetails();
 					}
 				}
+				if (LinkedNode && InputPins.IsValidIndex(PinIndex))
+				{
+					InputPins[PinIndex]->BreakAllPinLinks();
+					InputPins[PinIndex]->MakeLinkTo(LinkedNode->GetOutputPin());
+				}
+
 			}
 		}
+	}
 
-		NodeName = NewName;
-		NodeTitle = FText::FromString(NodeName);
-		UpdateNodeFields();
+	FString NewName = FeaturePtr->Name;	
+
+	if (NewName != NodeName && NodeType != ERigMapperFeatureType::Output && NodeType != ERigMapperFeatureType::NullOutput)
+	{
+		// Update inputs of all nodes directly connected to this node.
+		for (UEdGraphPin* OutPin : OutputPins) // It will always be one output pin.
+		{
+			for (UEdGraphPin* LinkedPin : OutPin->LinkedTo)
+			{
+				if (URigMapperDefinitionEditorGraphNode* LinkedNode = Cast<URigMapperDefinitionEditorGraphNode>(LinkedPin->GetOwningNode()))
+				{
+					LinkedNode->RelinkFeature(NodeName, NewName);
+				}
+			}
+		}
 	}
+	
+	NodeName = NewName;
+	NodeTitle = FText::FromString(NodeName);
+	UpdateNodeFields();	
 }
 
 #undef LOCTEXT_NAMESPACE
diff --git a/Engine/Plugins/Experimental/Animation/RigMapper/Source/RigMapperEditor/Private/RigMapperGraph/RigMapperDefinitionEditorGraphNode.h b/Engine/Plugins/Experimental/Animation/RigMapper/Source/RigMapperEditor/Private/RigMapperGraph/RigMapperDefinitionEditorGraphNode.h
--- a/Engine/Plugins/Experimental/Animation/RigMapper/Source/RigMapperEditor/Private/RigMapperGraph/RigMapperDefinitionEditorGraphNode.h
+++ b/Engine/Plugins/Experimental/Animation/RigMapper/Source/RigMapperEditor/Private/RigMapperGraph/RigMapperDefinitionEditorGraphNode.h
@@ -10,21 +10,10 @@
 
 #define UE_API RIGMAPPEREDITOR_API
 
+enum class ERigMapperFeatureType : uint8;
 class URigMapperDefinition;
 struct FRigMapperFeature;
 
-UENUM()
-enum class ERigMapperNodeType : uint8
-{
-	Input,
-	WeightedSum,
-	SDK,
-	Multiply,
-	Output,
-	NullOutput,
-	Invalid
-};
-
 /**
  * 
  */
@@ -39,27 +28,27 @@ public:
 		virtual TSharedPtr<SGraphNode> CreateNode(UEdGraphNode* Node) const override;
 	};
 
-	
+
 	// UEdGraphNode implementation
 	virtual FText GetNodeTitle(ENodeTitleType::Type TitleType) const override { return NodeTitle; }
 	UE_API virtual FLinearColor GetNodeBodyTintColor() const override;
 	UE_API virtual FLinearColor GetNodeTitleColor() const override;
 	UE_API virtual void GetNodeContextMenuActions(class UToolMenu* Menu, class UGraphNodeContextMenuContext* Context) const override;
 	UE_API virtual void AllocateDefaultPins() override;
-	UE_API void BeginDestroy() override;
 	// !UEdGraphNode implementation
 
-	UE_API void SetupNode(URigMapperDefinition* InDefinition, const FString& InNodeName, ERigMapperNodeType InNodeType);
+	UE_API void SetupNode(URigMapperDefinition* InDefinition, const FString& InNodeName, ERigMapperFeatureType InNodeType);
 	const FString& GetNodeName() const { return NodeName; }
-	
-	ERigMapperNodeType GetNodeType() const { return NodeType; }
+
+	ERigMapperFeatureType GetNodeType() const { return NodeType; }
 
 	UE_API FRigMapperFeature* GetNodeFeature() const;
-	
+
 	const FText& GetSubtitle() const { return NodeSubtitle; }
-	
+
 	UE_API UEdGraphPin* CreateInputPin();
 	UE_API UEdGraphPin* CreateOutputPin();
+	UE_API bool RemoveInputPin(UEdGraphPin* InPinToRemove);
 	const TArray<UEdGraphPin*>& GetInputPins() const { return InputPins; }
 	UEdGraphPin* GetOutputPin() const { return OutputPins.IsEmpty() ? nullptr : OutputPins[0]; };
 
@@ -72,7 +61,7 @@ public:
 	UE_API void GetRect(FVector2D& TopLeft, FVector2D& BottomRight) const;
 
 	/** Updates node from definition and makes it dirty */
-	UE_API void UpdateNodeFromFeature();
+	UE_API void UpdateNodeFromFeature(const FPropertyChangedEvent& InPropertyChangedEvent);
 
 	/** Updates node body */
 	UE_API void UpdateNodeFields();
@@ -98,16 +87,12 @@ protected:
 
 	/** Cached dimensions of this node (used for layout) */
 	FVector2D Margin = { 0, 0 };
-	
-	TWeakObjectPtr<URigMapperDefinition> Definition;
 
-	FRigMapperFeature* FeaturePtr = nullptr;
+	TWeakObjectPtr<URigMapperDefinition> Definition;
 
 	FString NodeName;
 
-	ERigMapperNodeType NodeType = ERigMapperNodeType::Invalid;
-
-	bool IsExternalOwner;
+	ERigMapperFeatureType NodeType;
 };
 
 #undef UE_API
diff --git a/Engine/Plugins/Experimental/Animation/RigMapper/Source/RigMapperEditor/Private/RigMapperGraph/RigMapperDefinitionEditorGraphSchema.cpp b/Engine/Plugins/Experimental/Animation/RigMapper/Source/RigMapperEditor/Private/RigMapperGraph/RigMapperDefinitionEditorGraphSchema.cpp
--- a/Engine/Plugins/Experimental/Animation/RigMapper/Source/RigMapperEditor/Private/RigMapperGraph/RigMapperDefinitionEditorGraphSchema.cpp
+++ b/Engine/Plugins/Experimental/Animation/RigMapper/Source/RigMapperEditor/Private/RigMapperGraph/RigMapperDefinitionEditorGraphSchema.cpp
@@ -18,11 +18,11 @@
 
 FSchemaAction_SpawnRigMapperNode::FSchemaAction_SpawnRigMapperNode() :
 	FEdGraphSchemaAction(),
-	NodeType(ERigMapperNodeType::NullOutput)
+	NodeType(ERigMapperFeatureType::NullOutput)
 {
 }
 
-FSchemaAction_SpawnRigMapperNode::FSchemaAction_SpawnRigMapperNode(ERigMapperNodeType InNodeType, const FText& InCategory, const FText& InDescription, const FText& InTooltip, const int32 InGrouping) :
+FSchemaAction_SpawnRigMapperNode::FSchemaAction_SpawnRigMapperNode(ERigMapperFeatureType InNodeType, const FText& InCategory, const FText& InDescription, const FText& InTooltip, const int32 InGrouping) :
 	FEdGraphSchemaAction(InCategory, InDescription, InTooltip, InGrouping),
 	NodeType(InNodeType)
 {
@@ -154,9 +154,9 @@ bool URigMapperDefinitionEditorGraphSchema::SupportsDropPinOnNode(UEdGraphNode*
 		return false;
 	}
 
-	if (URigMapperDefinitionEditorGraphNode* GraphNode = Cast<URigMapperDefinitionEditorGraphNode>(InTargetNode))
+	if (const URigMapperDefinitionEditorGraphNode* GraphNode = Cast<URigMapperDefinitionEditorGraphNode>(InTargetNode))
 	{
-		if (GraphNode->GetNodeType() == ERigMapperNodeType::Multiply || GraphNode->GetNodeType() == ERigMapperNodeType::WeightedSum)
+		if (GraphNode->GetNodeType() == ERigMapperFeatureType::Multiply || GraphNode->GetNodeType() == ERigMapperFeatureType::WeightedSum)
 		{
 			OutErrorMessage = LOCTEXT("RigMapperDefinitionEditor_CreateNewConnection_DropPinOnNode", "Create new input pin and connect.");
 			return true;
@@ -179,7 +179,7 @@ void URigMapperDefinitionEditorGraphSchema::GetGraphContextActions(FGraphContext
 
 	if (!FromPin || FromPin->Direction == EGPD_Input)
 	{
-		ContextMenuBuilder.AddAction(MakeShared<FSchemaAction_SpawnRigMapperNode>(ERigMapperNodeType::Input, Category,
+		ContextMenuBuilder.AddAction(MakeShared<FSchemaAction_SpawnRigMapperNode>(ERigMapperFeatureType::Input, Category,
 			LOCTEXT("RigMapperDefinitionEditor_AddInputNode", "Add Input Node"),
 			LOCTEXT("RigMapperDefinitionEditor_AddInputTooltip", "Spawn Input node (at the end of the input connector, if available)"),
 			/*Grouping=*/0));
@@ -187,30 +187,30 @@ void URigMapperDefinitionEditorGraphSchema::GetGraphContextActions(FGraphContext
 
 	if (!FromPin || FromPin->Direction == EGPD_Output)
 	{
-		ContextMenuBuilder.AddAction(MakeShared<FSchemaAction_SpawnRigMapperNode>(ERigMapperNodeType::Output, Category,
+		ContextMenuBuilder.AddAction(MakeShared<FSchemaAction_SpawnRigMapperNode>(ERigMapperFeatureType::Output, Category,
 			LOCTEXT("RigMapperDefinitionEditor_AddOutputNode", "Add Output Node"),
 			LOCTEXT("RigMapperDefinitionEditor_AddOutputNodeTooltip", "Spawn Output node (at the end of the output connector, if available)"),
 			/*Grouping=*/0));
 	}
 
-	ContextMenuBuilder.AddAction(MakeShared<FSchemaAction_SpawnRigMapperNode>(ERigMapperNodeType::Multiply, Category,
+	ContextMenuBuilder.AddAction(MakeShared<FSchemaAction_SpawnRigMapperNode>(ERigMapperFeatureType::Multiply, Category,
 		LOCTEXT("RigMapperDefinitionEditor_AddMultiplyNode", "Add Multiply Node"),
 		LOCTEXT("RigMapperDefinitionEditor_AddMultiplyNodeTooltip", "Spawn Multiply node (at the end of the connector, if available)"),
 		/*Grouping=*/0));
 
-	ContextMenuBuilder.AddAction(MakeShared<FSchemaAction_SpawnRigMapperNode>(ERigMapperNodeType::SDK, Category,
+	ContextMenuBuilder.AddAction(MakeShared<FSchemaAction_SpawnRigMapperNode>(ERigMapperFeatureType::SDK, Category,
 		LOCTEXT("RigMapperDefinitionEditor_AddSDKNode", "Add SDK Node"),
 		LOCTEXT("RigMapperDefinitionEditor_AddSDKNodeTooltip", "Spawn SDK node (at the end of the connector, if available)"),
 		/*Grouping=*/0));
 
-	ContextMenuBuilder.AddAction(MakeShared<FSchemaAction_SpawnRigMapperNode>(ERigMapperNodeType::WeightedSum, Category,
+	ContextMenuBuilder.AddAction(MakeShared<FSchemaAction_SpawnRigMapperNode>(ERigMapperFeatureType::WeightedSum, Category,
 		LOCTEXT("RigMapperDefinitionEditor_AddWeightedSumNode", "Add WeightedSum Node"),
 		LOCTEXT("RigMapperDefinitionEditor_AddWeightedSumNodeTooltip", "Spawn WeightedSum node (at the end of the connector, if available)"),
 		/*Grouping=*/0));
 
 	if (!FromPin)
 	{
-		ContextMenuBuilder.AddAction(MakeShared<FSchemaAction_SpawnRigMapperNode>(ERigMapperNodeType::NullOutput, Category,
+		ContextMenuBuilder.AddAction(MakeShared<FSchemaAction_SpawnRigMapperNode>(ERigMapperFeatureType::NullOutput, Category,
 			LOCTEXT("RigMapperDefinitionEditor_AddNullOutputNode", "Add Null Output Node"),
 			LOCTEXT("RigMapperDefinitionEditor_AddNullOutputNodeTooltip", "Spawn null output node"),
 			/*Grouping=*/0));
@@ -233,12 +233,12 @@ bool URigMapperDefinitionEditorGraphSchema::TryCreateConnection(UEdGraphPin* InA
 
 bool URigMapperDefinitionEditorGraphSchema::IsConnectionRelinkingAllowed(UEdGraphPin* InPin) const
 {
-	return Super::IsConnectionRelinkingAllowed(InPin);
+	return true;
 }
 
 const FPinConnectionResponse URigMapperDefinitionEditorGraphSchema::CanRelinkConnectionToPin(const UEdGraphPin* OldSourcePin, const UEdGraphPin* TargetPinCandidate) const
 {
-	return FPinConnectionResponse(CONNECT_RESPONSE_DISALLOW, TEXT(""));
+	return FPinConnectionResponse(CONNECT_RESPONSE_MAKE, TEXT(""));
 }
 
 bool URigMapperDefinitionEditorGraphSchema::CreatePromotedConnection(UEdGraphPin* A, UEdGraphPin* B) const
@@ -250,4 +250,9 @@ void URigMapperDefinitionEditorGraphSchema::SetPinBeingDroppedOnNode(UEdGraphPin
 {
 }
 
+bool URigMapperDefinitionEditorGraphSchema::TryRelinkConnectionTarget(UEdGraphPin* SourcePin, UEdGraphPin* OldTargetPin, UEdGraphPin* NewTargetPin, const TArray<UEdGraphNode*>& InSelectedGraphNodes) const
+{
+	return Super::TryRelinkConnectionTarget(SourcePin, OldTargetPin, NewTargetPin, InSelectedGraphNodes);
+}
+
 #undef LOCTEXT_NAMESPACE
diff --git a/Engine/Plugins/Experimental/Animation/RigMapper/Source/RigMapperEditor/Private/RigMapperGraph/RigMapperDefinitionEditorGraphSchema.h b/Engine/Plugins/Experimental/Animation/RigMapper/Source/RigMapperEditor/Private/RigMapperGraph/RigMapperDefinitionEditorGraphSchema.h
--- a/Engine/Plugins/Experimental/Animation/RigMapper/Source/RigMapperEditor/Private/RigMapperGraph/RigMapperDefinitionEditorGraphSchema.h
+++ b/Engine/Plugins/Experimental/Animation/RigMapper/Source/RigMapperEditor/Private/RigMapperGraph/RigMapperDefinitionEditorGraphSchema.h
@@ -10,7 +10,7 @@
 
 #define UE_API RIGMAPPEREDITOR_API
 
-enum class ERigMapperNodeType : uint8;
+enum class ERigMapperFeatureType : uint8;
 
 USTRUCT()
 struct FSchemaAction_SpawnRigMapperNode : public FEdGraphSchemaAction
@@ -19,13 +19,13 @@ struct FSchemaAction_SpawnRigMapperNode : public FEdGraphSchemaAction
 
 public:
 	FSchemaAction_SpawnRigMapperNode();
-	FSchemaAction_SpawnRigMapperNode(ERigMapperNodeType InNodeType, const FText& InCategory, const FText& InDescription, const FText& InTooltip, const int);
+	FSchemaAction_SpawnRigMapperNode(ERigMapperFeatureType InNodeType, const FText& InCategory, const FText& InDescription, const FText& InTooltip, const int InGrouping);
 
 	//~ Begin FEdGraphSchemaAction Interface
 	virtual UEdGraphNode* PerformAction(UEdGraph* ParentGraph, UEdGraphPin* FromPin, const FVector2f& Location, bool bSelectNewNode = true) override;
 	//~ End FEdGraphSchemaAction Interface
 private:
-	ERigMapperNodeType NodeType;
+	ERigMapperFeatureType NodeType;
 };
 
 /**
@@ -75,9 +75,12 @@ public:
 	UE_API bool CreatePromotedConnection(UEdGraphPin* InA, UEdGraphPin* InB) const override;
 	UEdGraphPin* DropPinOnNode(UEdGraphNode* InTargetNode, const FName& InSourcePinName, const FEdGraphPinType& InSourcePinType, EEdGraphPinDirection InSourcePinDirection) const override;
 	void SetPinBeingDroppedOnNode(UEdGraphPin* InSourcePin) const override;
+	UE_API bool TryRelinkConnectionTarget(UEdGraphPin* SourcePin, UEdGraphPin* OldTargetPin, UEdGraphPin* NewTargetPin, const TArray<UEdGraphNode*>& InSelectedGraphNodes) const override;
 	// End of UEdGraphSchema interface
 private:
 
+
+
 };
 
 #undef UE_API
diff --git a/Engine/Plugins/Experimental/Animation/RigMapper/Source/RigMapperEditor/Private/RigMapperGraph/SRigMapperDefinitionGraphEditor.cpp b/Engine/Plugins/Experimental/Animation/RigMapper/Source/RigMapperEditor/Private/RigMapperGraph/SRigMapperDefinitionGraphEditor.cpp
--- a/Engine/Plugins/Experimental/Animation/RigMapper/Source/RigMapperEditor/Private/RigMapperGraph/SRigMapperDefinitionGraphEditor.cpp
+++ b/Engine/Plugins/Experimental/Animation/RigMapper/Source/RigMapperEditor/Private/RigMapperGraph/SRigMapperDefinitionGraphEditor.cpp
@@ -67,20 +67,20 @@ void SRigMapperDefinitionGraphEditor::Tick(const FGeometry& AllottedGeometry, co
 
 		GraphObj->RequestRefreshLayout(false);
 
-		const FGraphPanelSelectionSet& Selection = GraphEditor->GetSelectedNodes();
+		//const FGraphPanelSelectionSet& Selection = GraphEditor->GetSelectedNodes();
 
-		TArray<URigMapperDefinitionEditorGraphNode*> SelectedNodes;
-		SelectedNodes.Reserve(Selection.Num());
+		//TArray<URigMapperDefinitionEditorGraphNode*> SelectedNodes;
+		//SelectedNodes.Reserve(Selection.Num());
 
-		for (UObject* Obj : Selection)
-		{
-			if (URigMapperDefinitionEditorGraphNode* Node = Cast<URigMapperDefinitionEditorGraphNode>(Obj))
-			{
-				SelectedNodes.Add(Node);
-			}
-		}
+		//for (UObject* Obj : Selection)
+		//{
+		//	if (URigMapperDefinitionEditorGraphNode* Node = Cast<URigMapperDefinitionEditorGraphNode>(Obj))
+		//	{
+		//		SelectedNodes.Add(Node);
+		//	}
+		//}
 
-		ZoomToFitNodes(SelectedNodes);
+		//ZoomToFitNodes(SelectedNodes);
 	}
 }
 
diff --git a/Engine/Plugins/Experimental/Animation/RigMapper/Source/RigMapperEditor/Private/RigMapperGraph/SRigMapperDefinitionGraphEditorNode.cpp b/Engine/Plugins/Experimental/Animation/RigMapper/Source/RigMapperEditor/Private/RigMapperGraph/SRigMapperDefinitionGraphEditorNode.cpp
--- a/Engine/Plugins/Experimental/Animation/RigMapper/Source/RigMapperEditor/Private/RigMapperGraph/SRigMapperDefinitionGraphEditorNode.cpp
+++ b/Engine/Plugins/Experimental/Animation/RigMapper/Source/RigMapperEditor/Private/RigMapperGraph/SRigMapperDefinitionGraphEditorNode.cpp
@@ -12,13 +12,7 @@ BEGIN_SLATE_FUNCTION_BUILD_OPTIMIZATION
 void SRigMapperDefinitionGraphEditorNode::Construct(const FArguments& InArgs, URigMapperDefinitionEditorGraphNode* InNode)
 {
 	GraphNode = InNode;
-	SetCursor(EMouseCursor::CardinalCross);
-
-	// if (!ContentWidget.IsValid())
-	// {
-	// 	ContentWidget = SNullWidget::NullWidget;
-	// }
-	
+	SetCursor(EMouseCursor::CardinalCross);	
 	UpdateGraphNode();
 }
 
diff --git a/Engine/Plugins/Experimental/Animation/RigMapper/Source/RigMapperEditor/Private/RigMapperGraph/SRigMapperDefinitionStructureView.cpp b/Engine/Plugins/Experimental/Animation/RigMapper/Source/RigMapperEditor/Private/RigMapperGraph/SRigMapperDefinitionStructureView.cpp
--- a/Engine/Plugins/Experimental/Animation/RigMapper/Source/RigMapperEditor/Private/RigMapperGraph/SRigMapperDefinitionStructureView.cpp
+++ b/Engine/Plugins/Experimental/Animation/RigMapper/Source/RigMapperEditor/Private/RigMapperGraph/SRigMapperDefinitionStructureView.cpp
@@ -64,7 +64,7 @@ void SRigMapperDefinitionStructureView::Construct(const FArguments& InArgs, URig
 END_SLATE_FUNCTION_BUILD_OPTIMIZATION
 
 
-TSharedPtr<FString> SRigMapperDefinitionStructureView::GetParentAndChildrenNodes(ERigMapperNodeType NodeType, TArray<TSharedPtr<FString>>& OutChildren)
+TSharedPtr<FString> SRigMapperDefinitionStructureView::GetParentAndChildrenNodes(ERigMapperFeatureType NodeType, TArray<TSharedPtr<FString>>& OutChildren)
 {
 	if (const TSharedPtr<FString>* ParentNode = ParentNodesMapping.Find(NodeType))
 	{
@@ -77,7 +77,7 @@ TSharedPtr<FString> SRigMapperDefinitionStructureView::GetParentAndChildrenNodes
 	return nullptr;
 }
 
-TArray<TSharedPtr<FString>>* SRigMapperDefinitionStructureView::GetChildrenNodes(ERigMapperNodeType NodeType)
+TArray<TSharedPtr<FString>>* SRigMapperDefinitionStructureView::GetChildrenNodes(ERigMapperFeatureType NodeType)
 {
 	if (const TSharedPtr<FString>* ParentNode = ParentNodesMapping.Find(NodeType))
 	{
@@ -86,7 +86,7 @@ TArray<TSharedPtr<FString>>* SRigMapperDefinitionStructureView::GetChildrenNodes
 	return nullptr;
 }
 
-const TArray<TSharedPtr<FString>>* SRigMapperDefinitionStructureView::GetChildrenNodes(ERigMapperNodeType NodeType) const
+const TArray<TSharedPtr<FString>>* SRigMapperDefinitionStructureView::GetChildrenNodes(ERigMapperFeatureType NodeType) const
 {
 	if (const TSharedPtr<FString>* ParentNode = ParentNodesMapping.Find(NodeType))
 	{
@@ -95,7 +95,7 @@ const TArray<TSharedPtr<FString>>* SRigMapperDefinitionStructureView::GetChildre
 	return nullptr;
 }
 
-bool SRigMapperDefinitionStructureView::SelectNode(const FString& NodeName, ERigMapperNodeType NodeType, bool bSelected)
+bool SRigMapperDefinitionStructureView::SelectNode(const FString& NodeName, ERigMapperFeatureType NodeType, bool bSelected)
 {
 	if (const TArray<TSharedPtr<FString>>* ChildrenNodes = GetChildrenNodes(NodeType))
 	{
@@ -126,13 +126,13 @@ void SRigMapperDefinitionStructureView::ClearSelection() const
 	TreeView->ClearSelection();
 }
 
-bool SRigMapperDefinitionStructureView::IsNodeOrChildSelected(ERigMapperNodeType NodeType, int32 ArrayIndex) const
+bool SRigMapperDefinitionStructureView::IsNodeOrChildSelected(ERigMapperFeatureType NodeType, int32 ArrayIndex) const
 {
-	if (NodeType == ERigMapperNodeType::Invalid)
+	if (NodeType == ERigMapperFeatureType::Invalid)
 	{
-		return IsNodeOrChildSelected(ERigMapperNodeType::Multiply, ArrayIndex) ||
-				IsNodeOrChildSelected(ERigMapperNodeType::WeightedSum, ArrayIndex) ||
-				IsNodeOrChildSelected(ERigMapperNodeType::SDK, ArrayIndex);
+		return IsNodeOrChildSelected(ERigMapperFeatureType::Multiply, ArrayIndex) ||
+				IsNodeOrChildSelected(ERigMapperFeatureType::WeightedSum, ArrayIndex) ||
+				IsNodeOrChildSelected(ERigMapperFeatureType::SDK, ArrayIndex);
 	}
 	
 	const TArray<TSharedPtr<FString>>& Selection = TreeView->GetSelectedItems();
@@ -166,26 +166,26 @@ void SRigMapperDefinitionStructureView::GenerateParentNodes()
 	ParentsAndChildrenNodes.Empty(NumFeatureTypes + NumNodeTypes);
 
 	ParentsAndChildrenNodes.Add(ParentNodesMapping.Add(
-		ERigMapperNodeType::Input, RootNodes.Add_GetRef(MakeShared<FString>(InputsNodeName))));
+		ERigMapperFeatureType::Input, RootNodes.Add_GetRef(MakeShared<FString>(InputsNodeName))));
 
 	TArray<TSharedPtr<FString>>& FeatureEntries = ParentsAndChildrenNodes.Add(ParentNodesMapping.Add(
-		ERigMapperNodeType::Invalid, RootNodes.Add_GetRef(MakeShared<FString>(FeaturesNodeName))));
+		ERigMapperFeatureType::Invalid, RootNodes.Add_GetRef(MakeShared<FString>(FeaturesNodeName))));
 	FeatureEntries.Reserve(NumNodeTypes);
 
 	ParentsAndChildrenNodes.Add(ParentNodesMapping.Add(
-		ERigMapperNodeType::Multiply, FeatureEntries.Add_GetRef(MakeShared<FString>(MultiplyNodeName))));
+		ERigMapperFeatureType::Multiply, FeatureEntries.Add_GetRef(MakeShared<FString>(MultiplyNodeName))));
 
 	ParentsAndChildrenNodes.Add(ParentNodesMapping.Add(
-		ERigMapperNodeType::WeightedSum, FeatureEntries.Add_GetRef(MakeShared<FString>(WsNodeName))));
+		ERigMapperFeatureType::WeightedSum, FeatureEntries.Add_GetRef(MakeShared<FString>(WsNodeName))));
 	
 	ParentsAndChildrenNodes.Add(ParentNodesMapping.Add(
-		ERigMapperNodeType::SDK, FeatureEntries.Add_GetRef(MakeShared<FString>(SdkNodeName))));
+		ERigMapperFeatureType::SDK, FeatureEntries.Add_GetRef(MakeShared<FString>(SdkNodeName))));
 
 	ParentsAndChildrenNodes.Add(ParentNodesMapping.Add(
-		ERigMapperNodeType::Output, RootNodes.Add_GetRef(MakeShared<FString>(OutputNodeName))));
+		ERigMapperFeatureType::Output, RootNodes.Add_GetRef(MakeShared<FString>(OutputNodeName))));
 
 	ParentsAndChildrenNodes.Add(ParentNodesMapping.Add(
-		ERigMapperNodeType::NullOutput, RootNodes.Add_GetRef(MakeShared<FString>(NullOutputNodeName))));
+		ERigMapperFeatureType::NullOutput, RootNodes.Add_GetRef(MakeShared<FString>(NullOutputNodeName))));
 
 	FilteredRootNodes = RootNodes;
 }
@@ -194,42 +194,42 @@ void SRigMapperDefinitionStructureView::GenerateChildrenNodes()
 {
 	if (Definition && !ParentsAndChildrenNodes.IsEmpty() && !ParentNodesMapping.IsEmpty())
 	{
-		TArray<TSharedPtr<FString>>* InputEntries = GetChildrenNodes(ERigMapperNodeType::Input);
+		TArray<TSharedPtr<FString>>* InputEntries = GetChildrenNodes(ERigMapperFeatureType::Input);
 		InputEntries->Reset(Definition->Inputs.Num());
 		for (const FString& Input : Definition->Inputs)
 		{
 			InputEntries->Add(MakeShared<FString>(Input));
 		}
 
-		TArray<TSharedPtr<FString>>* MultiplyEntries = GetChildrenNodes(ERigMapperNodeType::Multiply);
+		TArray<TSharedPtr<FString>>* MultiplyEntries = GetChildrenNodes(ERigMapperFeatureType::Multiply);
 		MultiplyEntries->Reset(Definition->Features.Multiply.Num());
 		for (const FRigMapperMultiplyFeature& Feature : Definition->Features.Multiply)
 		{
 			MultiplyEntries->Add(MakeShared<FString>(Feature.Name));
 		}
 
-		TArray<TSharedPtr<FString>>* WsEntries = GetChildrenNodes(ERigMapperNodeType::WeightedSum);
+		TArray<TSharedPtr<FString>>* WsEntries = GetChildrenNodes(ERigMapperFeatureType::WeightedSum);
 		WsEntries->Reset(Definition->Features.WeightedSums.Num());
 		for (const FRigMapperFeature& Feature : Definition->Features.WeightedSums)
 		{
 			WsEntries->Add(MakeShared<FString>(Feature.Name));
 		}
 
-		TArray<TSharedPtr<FString>>* SdkEntries = GetChildrenNodes(ERigMapperNodeType::SDK);
+		TArray<TSharedPtr<FString>>* SdkEntries = GetChildrenNodes(ERigMapperFeatureType::SDK);
 		SdkEntries->Reset(Definition->Features.SDKs.Num());
 		for (const FRigMapperSdkFeature& Feature : Definition->Features.SDKs)
 		{
 			SdkEntries->Add(MakeShared<FString>(Feature.Name));
 		}
 
-		TArray<TSharedPtr<FString>>* OutputEntries = GetChildrenNodes(ERigMapperNodeType::Output);
+		TArray<TSharedPtr<FString>>* OutputEntries = GetChildrenNodes(ERigMapperFeatureType::Output);
 		OutputEntries->Reset(Definition->Outputs.Num());
 		for (const TPair<FString, FString>& Output : Definition->Outputs)
 		{
 			OutputEntries->Add(MakeShared<FString>(Output.Key));
 		}
 
-		TArray<TSharedPtr<FString>>* NullOutputEntries = GetChildrenNodes(ERigMapperNodeType::NullOutput);
+		TArray<TSharedPtr<FString>>* NullOutputEntries = GetChildrenNodes(ERigMapperFeatureType::NullOutput);
 		NullOutputEntries->Reset(Definition->NullOutputs.Num());
 		for (const FString& NullOutput : Definition->NullOutputs)
 		{
@@ -263,14 +263,14 @@ void SRigMapperDefinitionStructureView::HandleTreeNodesSelectionChanged(TSharedP
 			return SelectedNames;
 		};
 
-		SelectedInputs = GetSelectedNodeNames(GetChildrenNodes(ERigMapperNodeType::Input));
+		SelectedInputs = GetSelectedNodeNames(GetChildrenNodes(ERigMapperFeatureType::Input));
 
-		SelectedFeatures = GetSelectedNodeNames(GetChildrenNodes(ERigMapperNodeType::Multiply));
-		SelectedFeatures.Append(GetSelectedNodeNames(GetChildrenNodes(ERigMapperNodeType::WeightedSum)));
-		SelectedFeatures.Append(GetSelectedNodeNames(GetChildrenNodes(ERigMapperNodeType::SDK)));
+		SelectedFeatures = GetSelectedNodeNames(GetChildrenNodes(ERigMapperFeatureType::Multiply));
+		SelectedFeatures.Append(GetSelectedNodeNames(GetChildrenNodes(ERigMapperFeatureType::WeightedSum)));
+		SelectedFeatures.Append(GetSelectedNodeNames(GetChildrenNodes(ERigMapperFeatureType::SDK)));
 	
-		SelectedOutputs = GetSelectedNodeNames(GetChildrenNodes(ERigMapperNodeType::Output));
-		SelectedNullOutputs = GetSelectedNodeNames(GetChildrenNodes(ERigMapperNodeType::NullOutput));
+		SelectedOutputs = GetSelectedNodeNames(GetChildrenNodes(ERigMapperFeatureType::Output));
+		SelectedNullOutputs = GetSelectedNodeNames(GetChildrenNodes(ERigMapperFeatureType::NullOutput));
 	}
 
 	if (OnSelectionChanged.IsBound())
diff --git a/Engine/Plugins/Experimental/Animation/RigMapper/Source/RigMapperEditor/Private/RigMapperGraph/SRigMapperDefinitionStructureView.h b/Engine/Plugins/Experimental/Animation/RigMapper/Source/RigMapperEditor/Private/RigMapperGraph/SRigMapperDefinitionStructureView.h
--- a/Engine/Plugins/Experimental/Animation/RigMapper/Source/RigMapperEditor/Private/RigMapperGraph/SRigMapperDefinitionStructureView.h
+++ b/Engine/Plugins/Experimental/Animation/RigMapper/Source/RigMapperEditor/Private/RigMapperGraph/SRigMapperDefinitionStructureView.h
@@ -9,7 +9,7 @@
 
 #define UE_API RIGMAPPEREDITOR_API
 
-enum class ERigMapperNodeType : uint8;
+enum class ERigMapperFeatureType : uint8;
 class SSearchBox;
 class URigMapperDefinition;
 
@@ -30,10 +30,10 @@ public:
 	/** Constructs this widget with InArgs */
 	UE_API void Construct(const FArguments& InArgs, URigMapperDefinition* InDefinition);
 
-	UE_API bool SelectNode(const FString& NodeName, ERigMapperNodeType NodeType, bool bSelected);
+	UE_API bool SelectNode(const FString& NodeName, ERigMapperFeatureType NodeType, bool bSelected);
 	UE_API void RebuildTree();
 	UE_API void ClearSelection() const;
-	UE_API bool IsNodeOrChildSelected(ERigMapperNodeType NodeType, int32 ArrayIndex) const;
+	UE_API bool IsNodeOrChildSelected(ERigMapperFeatureType NodeType, int32 ArrayIndex) const;
 	UE_API bool IsSelectionEmpty() const;
 	
 public:
@@ -50,9 +50,9 @@ private:
 	UE_API TSharedRef<ITableRow> OnGenerateTreeRow(TSharedPtr<FString> NodeName, const TSharedRef<STableViewBase>& TableViewBase);
 	UE_API void OnGetTreeNodeChildren(TSharedPtr<FString> NodeName, TArray<TSharedPtr<FString>>& Children);
 
-	UE_API TSharedPtr<FString> GetParentAndChildrenNodes(ERigMapperNodeType NodeType, TArray<TSharedPtr<FString>>& OutChildren);
-	UE_API TArray<TSharedPtr<FString>>* GetChildrenNodes(ERigMapperNodeType NodeType);
-	UE_API const TArray<TSharedPtr<FString>>* GetChildrenNodes(ERigMapperNodeType NodeType) const;
+	UE_API TSharedPtr<FString> GetParentAndChildrenNodes(ERigMapperFeatureType NodeType, TArray<TSharedPtr<FString>>& OutChildren);
+	UE_API TArray<TSharedPtr<FString>>* GetChildrenNodes(ERigMapperFeatureType NodeType);
+	UE_API const TArray<TSharedPtr<FString>>* GetChildrenNodes(ERigMapperFeatureType NodeType) const;
 	
 private:
 	static UE_API const int32 NumNodeTypes;
@@ -77,7 +77,7 @@ private:
 	TArray<TSharedPtr<FString>> FilteredRootNodes;
 	
 	TMap<TSharedPtr<FString>, TArray<TSharedPtr<FString>>> ParentsAndChildrenNodes;
-	TMap<ERigMapperNodeType, TSharedPtr<FString>> ParentNodesMapping;
+	TMap<ERigMapperFeatureType, TSharedPtr<FString>> ParentNodesMapping;
 };
 
 #undef UE_API
