[RigMapper] Added features for relinking, renaming and deleting nodes

diff --git a/Engine/Plugins/Experimental/Animation/RigMapper/Source/RigMapper/Private/RigMapperDefinition.cpp b/Engine/Plugins/Experimental/Animation/RigMapper/Source/RigMapper/Private/RigMapperDefinition.cpp
--- a/Engine/Plugins/Experimental/Animation/RigMapper/Source/RigMapper/Private/RigMapperDefinition.cpp
+++ b/Engine/Plugins/Experimental/Animation/RigMapper/Source/RigMapper/Private/RigMapperDefinition.cpp
@@ -445,7 +445,7 @@ bool URigMapperDefinition::LoadFromJsonString(const FString& JsonString)
 		SetDefinitionValid(false);
 	}
 
-	OnRigMapperDefinitionUpdated.Broadcast();
+	OnRigMapperDefinitionCreated.Broadcast();
 
 	return bValidated;
 }
@@ -895,14 +895,11 @@ bool FRigMapperMultiplyFeature::BakeInput(URigMapperLinkedDefinitions* LinkedDef
 	return OutBakedInput.Key.IsValid() && !bMissingIsNotNullOutput; // if we have a missing definition which is not defined as a NullOutput, flag a failure
 }
 
-void FRigMapperMultiplyFeature::RelinkFeature(const FString& InOldInput, const FString& InNewInput)
+void FRigMapperMultiplyFeature::RelinkFeature(const FString& InNewInput, const int32 InputIndex)
 {
-	for (int32 i = 0; i < Inputs.Num(); i++)
+	if (Inputs.IsValidIndex(InputIndex))
 	{
-		if (Inputs[i] == InOldInput)
-		{
-			Inputs[i] = InNewInput;
-		}
+		Inputs[InputIndex] = InNewInput;
 	}
 }
 
@@ -911,6 +908,27 @@ void FRigMapperMultiplyFeature::AddNewInputLink(const FString& InNewInput)
 	Inputs.Add(InNewInput);
 }
 
+void FRigMapperMultiplyFeature::RemoveInput(const FString& InInput)
+{
+	Inputs.RemoveAll([InInput](const FString& Other) { return Other == InInput; });
+}
+
+void FRigMapperMultiplyFeature::ClearInputs()
+{
+	Inputs.Empty();
+}
+
+void FRigMapperMultiplyFeature::RenameInput(const FString& InOldInput, const FString& InNewInput)
+{
+	for (int32 i = 0; i < Inputs.Num(); i++)
+	{
+		if (Inputs[i] == InOldInput)
+		{
+			Inputs[i] = InNewInput;
+		}
+	}
+}
+
 bool FRigMapperWsFeature::BakeInput(URigMapperLinkedDefinitions* LinkedDefinitions, int32 DefinitionIndex, FBakedInput& OutBakedInput) const
 {
 	const FString BakedInputName = FString::Printf(TEXT("%s:%d"), *Name, DefinitionIndex);
@@ -955,21 +973,20 @@ bool FRigMapperWsFeature::BakeInput(URigMapperLinkedDefinitions* LinkedDefinitio
 	return OutBakedInput.Key.IsValid() && !bMissingIsNotNullOutput; // if we have a missing definition which is not defined as a NullOutput, flag a failure
 }
 
-void FRigMapperWsFeature::RelinkFeature(const FString& InOldInput, const FString& InNewInput)
+void FRigMapperWsFeature::RelinkFeature(const FString& InNewInput, const int32 InputIndex)
 {
-	if (double* Weight = Inputs.Find(InOldInput))
+	if (OrderedInputs.IsValidIndex(InputIndex))
 	{
-		for (int32 i = 0; i < OrderedInputs.Num(); i++)
+		const FString OldInput = OrderedInputs[InputIndex].Name;
+		if (double* Weight = Inputs.Find(OldInput))
 		{
-			if (OrderedInputs[i].Name == InOldInput)
-			{
-				OrderedInputs[i].Name = InNewInput; // Preserves the pin ordering.
-			}
-		}		
-		double TempValue = MoveTemp(*Weight);
-		Inputs.Remove(InOldInput);
-		Inputs.Add(InNewInput, TempValue);
+			double TempValue = MoveTemp(*Weight);
+			Inputs.Remove(OldInput);
+			Inputs.Add(InNewInput, TempValue);
+			OrderedInputs[InputIndex].Name = InNewInput;
+		}
 	}
+
 }
 
 void FRigMapperWsFeature::AddNewInputLink(const FString& InNewInput)
@@ -1003,21 +1020,47 @@ void FRigMapperWsFeature::UpdateFromDetails()
 
 void FRigMapperWsFeature::RemoveInput(const FString& InInput)
 {
+	for (int32 i = OrderedInputs.Num() - 1; i >= 0; i--)
+	{
+		if (OrderedInputs[i].Name == InInput)
+		{
+			OrderedInputs.RemoveAt(i);
+		}
+	}
 	Inputs.Remove(InInput);
 }
 
 void FRigMapperWsFeature::ClearInputs()
 {
+	OrderedInputs.Empty();
 	Inputs.Empty();
 }
 
 void FRigMapperWsFeature::RenameInput(const FString& InOldInput, const FString& InNewInput)
 {
+	if (InOldInput.IsEmpty())
+	{
+		return;
+	}
+	for (FRigMapperWsInput& WsInput : OrderedInputs)
+	{
+		if (WsInput.Name == InOldInput)
+		{
+			WsInput.Name = InNewInput;
+		}
+	}
 	if (double* Weight = Inputs.Find(InOldInput))
 	{
 		double TempValue = MoveTemp(*Weight);
 		Inputs.Remove(InOldInput);
-		Inputs.Add(InNewInput, TempValue);
+		if (!Inputs.Contains(InNewInput))
+		{
+			Inputs.Add(InNewInput, TempValue);
+		}
+		else
+		{
+			Inputs[InNewInput] += TempValue;
+		}
 	}
 }
 
@@ -1053,7 +1096,7 @@ bool FRigMapperSdkFeature::BakeInput(URigMapperLinkedDefinitions* LinkedDefiniti
 	return OutBakedInput.Key.IsValid() && !bMissingIsNotNullOutput; // if we have a missing definition which is not defined as a NullOutput, flag a failure
 }
 
-void FRigMapperSdkFeature::RelinkFeature(const FString& InOldInput, const FString& InNewInput)
+void FRigMapperSdkFeature::RelinkFeature(const FString& InNewInput, const int32 InputIndex)
 {
 	Input = InNewInput;
 }
@@ -1063,6 +1106,27 @@ void FRigMapperSdkFeature::AddNewInputLink(const FString& InNewInput)
 	Input = InNewInput;
 }
 
+void FRigMapperSdkFeature::RemoveInput(const FString& InInput)
+{
+	if (Input == InInput)
+	{
+		Input = TEXT("");
+	}
+}
+
+void FRigMapperSdkFeature::ClearInputs()
+{
+	Input = TEXT("");
+}
+
+void FRigMapperSdkFeature::RenameInput(const FString& InOldInput, const FString& InNewInput)
+{
+	if (Input == InOldInput)
+	{
+		Input = InNewInput;
+	}
+}
+
 bool FRigMapperSdkFeature::BakeKeys(const FRigMapperSdkFeature& InSdk, const FRigMapperSdkFeature& OutSdk, TArray<FRigMapperSdkKey>& BakedKeys)
 {
 	// Create an array to backward evaluate InSdk (First Layer)
@@ -1319,7 +1383,7 @@ bool URigMapperLinkedDefinitions::AreLinkedDefinitionsValid() const
 	return bOk;
 }
 
-void FRigMapperOutputFeature::RelinkFeature(const FString& InOldInput, const FString& InNewInput)
+void FRigMapperOutputFeature::RelinkFeature(const FString& InNewInput, const int32 InputIndex)
 {
 	Input = InNewInput;
 }
@@ -1328,3 +1392,24 @@ void FRigMapperOutputFeature::AddNewInputLink(const FString& InNewInput)
 {
 	Input = InNewInput;
 }
+
+void FRigMapperOutputFeature::RemoveInput(const FString& InInput)
+{
+	if (Input == InInput)
+	{
+		Input = TEXT("");
+	}
+}
+
+void FRigMapperOutputFeature::ClearInputs()
+{
+	Input = TEXT("");
+}
+
+void FRigMapperOutputFeature::RenameInput(const FString& InOldInput, const FString& InNewInput)
+{
+	if (Input == InOldInput)
+	{
+		Input = InNewInput;
+	}
+}
diff --git a/Engine/Plugins/Experimental/Animation/RigMapper/Source/RigMapper/Private/RigMapperUtils.cpp b/Engine/Plugins/Experimental/Animation/RigMapper/Source/RigMapper/Private/RigMapperUtils.cpp
new file mode 100644
--- /dev/null
+++ b/Engine/Plugins/Experimental/Animation/RigMapper/Source/RigMapper/Private/RigMapperUtils.cpp
@@ -0,0 +1,162 @@
+// Copyright Epic Games, Inc. All Rights Reserved.
+#include "RigMapperUtils.h"
+#include "RigMapperDefinition.h"
+#include "Algo/Find.h"
+#include "Algo/MaxElement.h"
+
+const TMap<ERigMapperFeatureType, FString> FeatureTypePrefixMap = {
+	{ ERigMapperFeatureType::WeightedSum, TEXT("ws:") },
+	{ ERigMapperFeatureType::Multiply, TEXT("mul:") },
+	{ ERigMapperFeatureType::SDK, TEXT("sdk:") },
+	{ ERigMapperFeatureType::Input, TEXT("") },
+	{ ERigMapperFeatureType::Output, TEXT("") },
+	{ ERigMapperFeatureType::NullOutput, TEXT("") },
+	{ ERigMapperFeatureType::Invalid, TEXT("invalid:") }
+};
+
+FString FRigMapperUtils::GenerateUniqueFeatureName(const TWeakObjectPtr<URigMapperDefinition> InDefinition, FString InDesiredName, ERigMapperFeatureType InFeatureType)
+{
+	TSet<FString> ExistingNames;
+	if (InFeatureType == ERigMapperFeatureType::Output || InFeatureType == ERigMapperFeatureType::NullOutput)
+	{
+		InDefinition.Pin()->Outputs.GetKeys(ExistingNames);
+		ExistingNames.Append(InDefinition.Pin()->NullOutputs);
+		bool bTrimmed = false;
+		if (InDesiredName.RemoveFromStart(FeatureTypePrefixMap[ERigMapperFeatureType::Multiply]))
+		{
+			bTrimmed = true;
+		}
+		if (!bTrimmed && InDesiredName.RemoveFromStart(FeatureTypePrefixMap[ERigMapperFeatureType::WeightedSum]))
+		{
+			bTrimmed = true;
+		}
+		if (!bTrimmed)
+		{
+			InDesiredName.RemoveFromStart(FeatureTypePrefixMap[ERigMapperFeatureType::SDK]);
+		}
+	}
+	else
+	{
+		ExistingNames.Append(InDefinition.Pin()->Inputs);
+		for (const FRigMapperFeature& Feature : InDefinition.Pin()->Features.Multiply)
+		{
+			ExistingNames.Add(Feature.Name);
+		}
+		for (const FRigMapperFeature& Feature : InDefinition.Pin()->Features.WeightedSums)
+		{
+			ExistingNames.Add(Feature.Name);
+		}
+		for (const FRigMapperFeature& Feature : InDefinition.Pin()->Features.SDKs)
+		{
+			ExistingNames.Add(Feature.Name);
+		}
+
+		InDesiredName = [&InDesiredName, InFeatureType]()
+			{
+				const FString& Prefix = FeatureTypePrefixMap[InFeatureType];
+				if (InDesiredName.IsEmpty() || InDesiredName.StartsWith(Prefix))
+				{
+					return InDesiredName;
+				}
+				if ((InFeatureType != ERigMapperFeatureType::Multiply && InDesiredName.RemoveFromStart(FeatureTypePrefixMap[ERigMapperFeatureType::Multiply])) ||
+					(InFeatureType != ERigMapperFeatureType::WeightedSum && InDesiredName.RemoveFromStart(FeatureTypePrefixMap[ERigMapperFeatureType::WeightedSum])) ||
+					(InFeatureType != ERigMapperFeatureType::SDK && InDesiredName.RemoveFromStart(FeatureTypePrefixMap[ERigMapperFeatureType::SDK])))
+				{
+					return Prefix + InDesiredName;
+				}
+				return Prefix + InDesiredName;
+			}();
+	}
+	if (!InDesiredName.IsEmpty() && !ExistingNames.Contains(InDesiredName))
+	{
+		return InDesiredName;
+	}
+
+	if (InDesiredName.IsEmpty())
+	{
+		InDesiredName = FeatureTypePrefixMap[InFeatureType] + TEXT("node");
+	}
+
+	// If the desired name is unique, just use that.
+	if (Algo::Find(ExistingNames, InDesiredName) == nullptr)
+	{
+		return InDesiredName;
+	}
+
+	// Create a unique name
+	FName UniqueName = [&InDesiredName, InFeatureType]()
+		{
+			int32 Len = InDesiredName.Len();
+			while (Len > 0 && FChar::IsDigit(InDesiredName[Len - 1]))
+			{
+				--Len;
+			}		
+
+			int32 Number;
+			if (!LexTryParseString(Number, *InDesiredName.RightChop(Len)))
+			{
+				return FName(*InDesiredName.Left(Len));
+			}
+
+			if (Len > 0 && InDesiredName[Len - 1] == TCHAR('_'))
+			{
+				--Len;
+			}
+
+			return FName(*InDesiredName.Left(Len), ++Number);
+		}();
+
+	const int32 MinNumDigitsInPostfix = 1;
+	const auto IsUniqueLambda = [MinNumDigitsInPostfix](TArray<FString> Others, const FName& InName)
+		{
+			const FString NameString = InName.ToString();
+
+			const FString& PlainNameString = InName.GetPlainNameString();
+			const int32 Number = NAME_INTERNAL_TO_EXTERNAL(InName.GetNumber());
+			const FString PossibleNameString = FString::Printf(TEXT("%s_%0*d"), *PlainNameString, MinNumDigitsInPostfix, Number);
+
+			for (const FString& Other : Others)
+			{
+				if (NameString == Other ||
+					PossibleNameString == Other)
+				{
+					return false;
+				}
+			}
+
+			return true;
+		};
+
+	// Only test names that start with the plain name string to optimize for performance
+	TArray<FString> FilteredExistingNames;
+	Algo::TransformIf(ExistingNames, FilteredExistingNames,
+		[UniqueNameString = UniqueName.GetPlainNameString()](const FString& ExistingName)
+		{
+			return ExistingName.StartsWith(UniqueNameString);
+		},
+		[](const FString& ExistingName)
+		{
+			return ExistingName;
+		});
+
+	// Sort in descending order to optimize for performance
+	Algo::Sort(FilteredExistingNames, [](const FString& NameA, const FString& NameB)
+		{
+			return NameA > NameB;
+		});
+
+	while (!IsUniqueLambda(FilteredExistingNames, UniqueName))
+	{
+		UniqueName.SetNumber(UniqueName.GetNumber() + 1);
+	}
+
+	const int32 UniqueNumber = NAME_INTERNAL_TO_EXTERNAL(UniqueName.GetNumber());
+	if (UniqueNumber == 0)
+	{
+		return UniqueName.ToString();
+	}
+	else
+	{
+		return FString::Printf(TEXT("%s_%0*d"), *UniqueName.GetPlainNameString(), MinNumDigitsInPostfix, UniqueNumber);
+	}
+}
diff --git a/Engine/Plugins/Experimental/Animation/RigMapper/Source/RigMapper/Public/RigMapperDefinition.h b/Engine/Plugins/Experimental/Animation/RigMapper/Source/RigMapper/Public/RigMapperDefinition.h
--- a/Engine/Plugins/Experimental/Animation/RigMapper/Source/RigMapper/Public/RigMapperDefinition.h
+++ b/Engine/Plugins/Experimental/Animation/RigMapper/Source/RigMapper/Public/RigMapperDefinition.h
@@ -55,7 +55,7 @@ public:
 
 	virtual bool BakeInput(URigMapperLinkedDefinitions* LinkedDefinitions, int32 DefinitionIndex, FBakedInput& OutBakedInput) const { return false; }
 
-	virtual void RelinkFeature(const FString& InOldInput, const FString& InNewInput) {}
+	virtual void RelinkFeature(const FString& InNewInput, const int32 InputIndex) {}
 
 	virtual void AddNewInputLink(const FString& InNewInput) {}
 
@@ -103,10 +103,14 @@ public:
 
 	UE_API virtual bool BakeInput(URigMapperLinkedDefinitions* LinkedDefinitions, int32 DefinitionIndex, FBakedInput& OutBakedInput) const override;
 
-	UE_API virtual void RelinkFeature(const FString& InOldInput, const FString& InNewInput) override;
+	UE_API virtual void RelinkFeature(const FString& InNewInput, const int32 InputIndex) override;
 	
 	UE_API virtual void AddNewInputLink(const FString& InNewInput) override;
 
+	UE_API virtual void RemoveInput(const FString& InInput) override;
+	UE_API virtual void ClearInputs() override;
+	UE_API virtual void RenameInput(const FString& InOldInput, const FString& InNewInput) override;
+
 	bool operator==(const FRigMapperMultiplyFeature& InOther) const
 	{
 		return FRigMapperFeature::operator==(InOther) && Inputs == InOther.Inputs;
@@ -186,7 +190,7 @@ public:
 
 	UE_API virtual bool BakeInput(URigMapperLinkedDefinitions* LinkedDefinitions, int32 DefinitionIndex, FBakedInput& OutBakedInput) const override;
 
-	UE_API virtual void RelinkFeature(const FString& InOldInput, const FString& InNewInput) override;
+	UE_API virtual void RelinkFeature(const FString& InNewInput, const int32 InputIndex) override;
 
 	UE_API virtual void AddNewInputLink(const FString& InNewInput) override;
 
@@ -263,10 +267,14 @@ public:
 
 	UE_API virtual bool BakeInput(URigMapperLinkedDefinitions* LinkedDefinitions, int32 DefinitionIndex, FBakedInput& OutBakedInput) const override;
 
-	UE_API virtual void RelinkFeature(const FString& InOldInput, const FString& InNewInput) override;
+	UE_API virtual void RelinkFeature(const FString& InNewInput, const int32 InputIndex) override;
 
 	UE_API virtual void AddNewInputLink(const FString& InNewInput) override;
 
+	UE_API virtual void RemoveInput(const FString& InInput) override;
+	UE_API virtual void ClearInputs() override;
+	UE_API virtual void RenameInput(const FString& InOldInput, const FString& InNewInput) override;
+
 	static UE_API bool BakeKeys(const FRigMapperSdkFeature& InSdk, const FRigMapperSdkFeature& OutSdk, TArray<FRigMapperSdkKey>& BakedKeys);
 
 	bool operator==(const FRigMapperSdkFeature& InOther) const
@@ -293,10 +301,14 @@ public:
 
 	virtual void GetInputs(TArray<FString>& OutInputs) const override { OutInputs = { Input }; };
 
-	UE_API virtual void RelinkFeature(const FString& InOldInput, const FString& InNewInput) override;
+	UE_API virtual void RelinkFeature(const FString& InNewInput, const int32 InputIndex) override;
 
 	UE_API virtual void AddNewInputLink(const FString& InNewInput) override;
 
+	UE_API virtual void RemoveInput(const FString& InInput) override;
+	UE_API virtual void ClearInputs() override;
+	UE_API virtual void RenameInput(const FString& InOldInput, const FString& InNewInput) override;
+
 	bool operator==(const FRigMapperOutputFeature& InOther) const
 	{
 		return FRigMapperFeature::operator==(InOther) && Input == InOther.Input;
@@ -315,7 +327,7 @@ public:
 
 	virtual ERigMapperFeatureType GetFeatureType() const override { return ERigMapperFeatureType::NullOutput; }
 
-	UE_API virtual void RelinkFeature(const FString& InOldInput, const FString& InNewInput) override {};
+	UE_API virtual void RelinkFeature(const FString& InNewInput, const int32 InputIndex) override {};
 
 	UE_API virtual void AddNewInputLink(const FString& InNewInput) override {};
 };
@@ -362,6 +374,7 @@ public:
 };
 
 
+DECLARE_MULTICAST_DELEGATE(FOnRigMapperDefinitionCreated);
 DECLARE_MULTICAST_DELEGATE(FOnRigMapperDefinitionUpdated);
 
 UCLASS(MinimalAPI, BlueprintType)
@@ -371,6 +384,8 @@ class URigMapperDefinition : public UDataAsset
 
 public:
 	// Delegate to notify listeners that we have loaded definitions
+	FOnRigMapperDefinitionCreated OnRigMapperDefinitionCreated;
+	// Delegate to notify listeners that the definition is modified (features added or deleted)
 	FOnRigMapperDefinitionUpdated OnRigMapperDefinitionUpdated;
 
 	UPROPERTY(EditAnywhere, BlueprintReadWrite, Category="Animation|Rig Mapper")
@@ -422,7 +437,6 @@ public:
 	 */
 	UE_API FRigMapperFeature* GetFeature(const FString& InFeatureName, const ERigMapperFeatureType InFeatureType);
 
-	
 private:
 	bool LoadInputsFromJsonObject(const TSharedPtr<FJsonObject>& JsonObject);
 	bool LoadFeaturesFromJsonObject(const TSharedPtr<FJsonObject>& JsonObject, TArray<FString>& FeatureNames);
diff --git a/Engine/Plugins/Experimental/Animation/RigMapper/Source/RigMapper/Public/RigMapperUtils.h b/Engine/Plugins/Experimental/Animation/RigMapper/Source/RigMapper/Public/RigMapperUtils.h
new file mode 100644
--- /dev/null
+++ b/Engine/Plugins/Experimental/Animation/RigMapper/Source/RigMapper/Public/RigMapperUtils.h
@@ -0,0 +1,26 @@
+// Copyright Epic Games, Inc. All Rights Reserved.
+
+#pragma once
+
+#include "CoreMinimal.h"
+
+#define UE_API RIGMAPPER_API
+
+class URigMapperDefinition;
+enum class ERigMapperFeatureType : uint8;
+
+class UE_API FRigMapperUtils
+{
+public:
+	/**
+	 * Generates a unique name given a base one and a list of existing ones from the Definition, by appending a type prefix and an index suffix.
+	 *
+	 * @param InDefinition		Definition which contains all the features to compare with.
+	 * @param InDesiredName		The desired Name. Note, generates '{type name}:node' if desired name is empty.
+	 * @param InFeatureType		Type of the feature being generated to include/replace as a prefix.
+	 * @return					Returns the unique name.
+	 */
+	static FString GenerateUniqueFeatureName(const TWeakObjectPtr<URigMapperDefinition> InDefinition, FString InDesiredName, ERigMapperFeatureType InFeatureType);
+};
+
+#undef UE_API
\ No newline at end of file
diff --git a/Engine/Plugins/Experimental/Animation/RigMapper/Source/RigMapperEditor/Private/RigMapperDefinitionEditorToolkit.cpp b/Engine/Plugins/Experimental/Animation/RigMapper/Source/RigMapperEditor/Private/RigMapperDefinitionEditorToolkit.cpp
--- a/Engine/Plugins/Experimental/Animation/RigMapper/Source/RigMapperEditor/Private/RigMapperDefinitionEditorToolkit.cpp
+++ b/Engine/Plugins/Experimental/Animation/RigMapper/Source/RigMapperEditor/Private/RigMapperDefinitionEditorToolkit.cpp
@@ -85,6 +85,15 @@ void FRigMapperDefinitionEditorToolkit::Initialize(URigMapperDefinition* InDefin
 	FRigMapperEditorModule::RegisterRigMapperDefinitionToolbarEntries();
 }
 
+void FRigMapperDefinitionEditorToolkit::OnClose()
+{
+	if (Definition)
+	{
+		Definition->OnRigMapperDefinitionCreated.RemoveAll(this);
+		Definition->OnRigMapperDefinitionUpdated.RemoveAll(this);
+	}
+}
+
 void FRigMapperDefinitionEditorToolkit::RegisterTabSpawners(const TSharedRef<FTabManager>& InTabManager)
 {
 	WorkspaceMenuCategory = InTabManager->AddLocalWorkspaceMenuCategory(LOCTEXT("RigMapperDefinitionEditorTabGroup", "Rig Mapper Definition Editor"));
@@ -142,8 +151,11 @@ TSharedRef<SDockTab> FRigMapperDefinitionEditorToolkit::SpawnGraphTab(const FSpa
 
 	GraphEditor->OnSelectionChanged.BindRaw(this, &FRigMapperDefinitionEditorToolkit::HandleGraphSelectionChanged);
 
-	Definition->OnRigMapperDefinitionUpdated.AddRaw(this, &FRigMapperDefinitionEditorToolkit::HandleRigMapperDefinitionLoaded);
-	
+	if (Definition)
+	{
+		Definition->OnRigMapperDefinitionCreated.AddRaw(this, &FRigMapperDefinitionEditorToolkit::HandleRigMapperDefinitionLoaded);
+		Definition->OnRigMapperDefinitionUpdated.AddRaw(this, &FRigMapperDefinitionEditorToolkit::HandleRigMapperDefinitionUpdated);
+	}
 	TSharedRef<SDockTab> DockTab = SNew(SDockTab)
 	.TabColorScale(GetTabColorScale())
 	.Label(LOCTEXT("RigMapperDefinitionEditorToolkitGraphTab", "Graph"))
@@ -174,6 +186,14 @@ void FRigMapperDefinitionEditorToolkit::HandleRigMapperDefinitionLoaded()
 	}	
 }
 
+void FRigMapperDefinitionEditorToolkit::HandleRigMapperDefinitionUpdated()
+{
+	if (StructureView)
+	{
+		StructureView->RebuildTree();
+	}
+}
+
 TSharedRef<IDetailCustomization> FRigMapperDefinitionEditorToolkit::FDetailsViewCustomization::MakeInstance()
 {
 	return MakeShared<FDetailsViewCustomization>();
@@ -371,14 +391,10 @@ void FRigMapperDefinitionEditorToolkit::HandleGraphSelectionChanged(const TSet<U
 			}					
 		}
 	}
-	
-	//DetailsView->ForceRefresh();
 }
 
 void FRigMapperDefinitionEditorToolkit::HandleStructureSelectionChanged(ESelectInfo::Type SelectInfo, TArray<FString> SelectedInputs, TArray<FString> SelectedFeatures, TArray<FString> SelectedOutputs, TArray<FString> SelectedNullOutputs)
 {
-	//DetailsView->ForceRefresh();
-
 	if (SelectInfo != ESelectInfo::Direct)
 	{
 		GraphEditor->SelectNodes(SelectedInputs, SelectedFeatures, SelectedOutputs, SelectedNullOutputs);
@@ -388,46 +404,59 @@ void FRigMapperDefinitionEditorToolkit::HandleStructureSelectionChanged(ESelectI
 void FRigMapperDefinitionEditorToolkit::SetFeatureDetailsView(const URigMapperDefinitionEditorGraphNode* InNode)
 {
 	TObjectPtr<URigMapperDefinitionEditorGraphNode> Node = ConstCast<const URigMapperDefinitionEditorGraphNode>(InNode);
-	if (!Node || !Node->GetNodeFeature())
+	if (!Node)
 	{
 		return;
 	}
-
+	const FRigMapperFeature* FeaturePtr = Node->GetNodeFeature();
+	if (!FeaturePtr)
+	{
+		return;
+	}
+	FString CurrentFeatureName = FeaturePtr->Name;
+	TSharedPtr<FStructOnScope> StructContainer = nullptr;
 	if (Node->GetNodeType() == ERigMapperFeatureType::SDK)
 	{
-		const TSharedPtr<FStructOnScope> SdkStructure = MakeShared<FStructOnScope>(FRigMapperSdkFeature::StaticStruct(), reinterpret_cast<uint8*>(Node->GetNodeFeature()));
-		FeatureDetailsView->SetStructureData(SdkStructure);
+		StructContainer = MakeShared<FStructOnScope>(FRigMapperSdkFeature::StaticStruct(), reinterpret_cast<uint8*>(Node->GetNodeFeature()));
+		FeatureDetailsView->SetStructureData(StructContainer);
 	}
 	else if (Node->GetNodeType() == ERigMapperFeatureType::WeightedSum)
 	{
-		const TSharedPtr<FStructOnScope> WsStructure = MakeShared<FStructOnScope>(FRigMapperWsFeature::StaticStruct(), reinterpret_cast<uint8*>(Node->GetNodeFeature()));
-		FeatureDetailsView->SetStructureData(WsStructure);
+		StructContainer = MakeShared<FStructOnScope>(FRigMapperWsFeature::StaticStruct(), reinterpret_cast<uint8*>(Node->GetNodeFeature()));
 	}
 	else if (Node->GetNodeType() == ERigMapperFeatureType::Multiply)
 	{
-		const TSharedPtr<FStructOnScope> MultiplyStructure = MakeShared<FStructOnScope>(FRigMapperMultiplyFeature::StaticStruct(), reinterpret_cast<uint8*>(Node->GetNodeFeature()));
-		FeatureDetailsView->SetStructureData(MultiplyStructure);
+		StructContainer = MakeShared<FStructOnScope>(FRigMapperMultiplyFeature::StaticStruct(), reinterpret_cast<uint8*>(Node->GetNodeFeature()));
+		FeatureDetailsView->SetStructureData(StructContainer);
 	}
 	else if (Node->GetNodeType() == ERigMapperFeatureType::Output)
 	{
-		const TSharedPtr<FStructOnScope> OutputStructure = MakeShared<FStructOnScope>(FRigMapperOutputFeature::StaticStruct(), reinterpret_cast<uint8*>(Node->GetNodeFeature()));
-		FeatureDetailsView->SetStructureData(OutputStructure);
+		StructContainer = MakeShared<FStructOnScope>(FRigMapperOutputFeature::StaticStruct(), reinterpret_cast<uint8*>(Node->GetNodeFeature()));
+		FeatureDetailsView->SetStructureData(StructContainer);
 	}
 	else if (Node->GetNodeType() == ERigMapperFeatureType::Input || Node->GetNodeType() == ERigMapperFeatureType::NullOutput)
 	{
-		const TSharedPtr<FStructOnScope> FeatureStructure = MakeShared<FStructOnScope>(FRigMapperFeature::StaticStruct(), reinterpret_cast<uint8*>(Node->GetNodeFeature()));
-		FeatureDetailsView->SetStructureData(FeatureStructure);
-	}			
+		StructContainer = MakeShared<FStructOnScope>(FRigMapperFeature::StaticStruct(), reinterpret_cast<uint8*>(Node->GetNodeFeature()));
+		FeatureDetailsView->SetStructureData(StructContainer);
+	}
 	FeatureDetailsView->GetOnFinishedChangingPropertiesDelegate().Clear();
-	FeatureDetailsView->GetOnFinishedChangingPropertiesDelegate().AddRaw(this, &FRigMapperDefinitionEditorToolkit::HandleFinishedChangingFeatureProperties, TWeakObjectPtr<URigMapperDefinitionEditorGraphNode>(Node.Get()));	
+	if (StructContainer)
+	{
+		FeatureDetailsView->SetStructureData(StructContainer);
+		FeatureDetailsView->GetOnFinishedChangingPropertiesDelegate().AddRaw(this, &FRigMapperDefinitionEditorToolkit::HandleFinishedChangingFeatureProperties, TWeakObjectPtr<URigMapperDefinitionEditorGraphNode>(Node.Get()), StructContainer.ToWeakPtr());
+	}
 }
 
-void FRigMapperDefinitionEditorToolkit::HandleFinishedChangingFeatureProperties(const FPropertyChangedEvent& InPropertyChangedEvent, TWeakObjectPtr<URigMapperDefinitionEditorGraphNode> InNode)
+void FRigMapperDefinitionEditorToolkit::HandleFinishedChangingFeatureProperties(const FPropertyChangedEvent& InPropertyChangedEvent, TWeakObjectPtr<URigMapperDefinitionEditorGraphNode> InNode, TWeakPtr<FStructOnScope> InStructContainer)
 {
-	if (InNode.IsValid())
-	{
-		InNode->UpdateNodeFromFeature(InPropertyChangedEvent);
+	if (InNode.IsValid() && GraphEditor.IsValid())
+	{		
+		InNode->UpdateNodeFromFeature(InPropertyChangedEvent, InStructContainer);
 		GraphEditor->RefreshGraphNode(InNode.Get());
+		if (StructureView)
+		{
+			StructureView->RebuildTree();
+		}
 	}	
 
 }
diff --git a/Engine/Plugins/Experimental/Animation/RigMapper/Source/RigMapperEditor/Private/RigMapperDefinitionEditorToolkit.h b/Engine/Plugins/Experimental/Animation/RigMapper/Source/RigMapperEditor/Private/RigMapperDefinitionEditorToolkit.h
--- a/Engine/Plugins/Experimental/Animation/RigMapper/Source/RigMapperEditor/Private/RigMapperDefinitionEditorToolkit.h
+++ b/Engine/Plugins/Experimental/Animation/RigMapper/Source/RigMapperEditor/Private/RigMapperDefinitionEditorToolkit.h
@@ -15,6 +15,7 @@ class URigMapperDefinitionEditorGraphNode;
 class SRigMapperDefinitionStructureView;
 class SRigMapperDefinitionGraphEditor;
 class IStructureDetailsView;
+class FStructOnScope;
 
 /**
  * The toolkit for the URigMapperDefinition asset editor
@@ -46,6 +47,9 @@ public:
 		// IDetailCustomization interface
 		virtual void CustomizeDetails(IDetailLayoutBuilder& DetailLayout) override;
 	};
+
+protected:
+	void OnClose() override;
 	
 private:
 	UE_API TSharedRef<SDockTab> SpawnGraphTab(const FSpawnTabArgs& SpawnTabArgs);
@@ -54,8 +58,9 @@ private:
 
 	UE_API bool HandleIsPropertyVisible(const FPropertyAndParent& PropertyAndParent);
 	UE_API void HandleFinishedChangingProperties(const FPropertyChangedEvent& PropertyChangedEvent);
-	UE_API void HandleFinishedChangingFeatureProperties(const FPropertyChangedEvent& PropertyChangedEvent, TWeakObjectPtr<URigMapperDefinitionEditorGraphNode> InNode);
+	UE_API void HandleFinishedChangingFeatureProperties(const FPropertyChangedEvent& PropertyChangedEvent, TWeakObjectPtr<URigMapperDefinitionEditorGraphNode> InNode, TWeakPtr<FStructOnScope> InStructContainer);
 	UE_API void HandleRigMapperDefinitionLoaded();
+	UE_API void HandleRigMapperDefinitionUpdated();
 	UE_API void HandleGraphSelectionChanged(const TSet<UObject*>& Nodes);
 	
 	UE_API void HandleStructureSelectionChanged(ESelectInfo::Type SelectInfo, TArray<FString> SelectedInputs, TArray<FString> SelectedFeatures, TArray<FString> SelectedOutputs, TArray<FString> SelectedNullOutputs);
@@ -74,6 +79,7 @@ private:
 	TSharedPtr<IDetailsView> DetailsView;
 	TSharedPtr<IStructureDetailsView> FeatureDetailsView;
 	TSharedPtr<SRigMapperDefinitionStructureView> StructureView;
+
 };
 
 #undef UE_API
diff --git a/Engine/Plugins/Experimental/Animation/RigMapper/Source/RigMapperEditor/Private/RigMapperGraph/RigMapperDefinitionEditorGraph.cpp b/Engine/Plugins/Experimental/Animation/RigMapper/Source/RigMapperEditor/Private/RigMapperGraph/RigMapperDefinitionEditorGraph.cpp
--- a/Engine/Plugins/Experimental/Animation/RigMapper/Source/RigMapperEditor/Private/RigMapperGraph/RigMapperDefinitionEditorGraph.cpp
+++ b/Engine/Plugins/Experimental/Animation/RigMapper/Source/RigMapperEditor/Private/RigMapperGraph/RigMapperDefinitionEditorGraph.cpp
@@ -4,6 +4,7 @@
 #include "RigMapperDefinition.h"
 #include "RigMapperDefinitionEditorGraphNode.h"
 #include "RigMapperDefinitionEditorGraphSchema.h"
+#include "RigMapperUtils.h"
 
 #include UE_INLINE_GENERATED_CPP_BY_NAME(RigMapperDefinitionEditorGraph)
 
@@ -35,7 +36,7 @@ void URigMapperDefinitionEditorGraph::ConstructNodes()
 	
 	for (const TPair<FString, FString>& Output : Definition->Outputs)
 	{
-		CreateGraphNodesRec(Definition, Output.Key, true);
+		ConstructGraphNodesRec(Definition, Output.Key, true);
 	}
 
 	// Generate nodes not related to any output
@@ -45,21 +46,21 @@ void URigMapperDefinitionEditorGraph::ConstructNodes()
 	{
 		if (!FeatureNodes.Contains(Feature))
 		{
-			CreateGraphNodesRec(Definition, Feature, false);
+			ConstructGraphNodesRec(Definition, Feature, false);
 		}
 	}
 	for (const FString& Input : Definition->Inputs)
 	{
 		if (!InputNodes.Contains(Input))
 		{
-			CreateGraphNode(Input, ERigMapperFeatureType::Input);
+			ConstructGraphNode(Input, ERigMapperFeatureType::Input);
 		}
 	}
 	for (const FString& NullOutput : Definition->NullOutputs)
 	{
 		if (!NullOutputNodes.Contains(NullOutput))
 		{
-			CreateGraphNode(NullOutput, ERigMapperFeatureType::NullOutput);
+			ConstructGraphNode(NullOutput, ERigMapperFeatureType::NullOutput);
 		}
 	}
 }
@@ -179,7 +180,7 @@ bool URigMapperDefinitionEditorGraph::RemoveNodeFromMap(const FString& InNodeNam
 	return true;
 }
 
-URigMapperDefinitionEditorGraphNode* URigMapperDefinitionEditorGraph::CreateGraphNodesRec(URigMapperDefinition* Definition, const FString& NodeName, bool bIsOutputNode)
+URigMapperDefinitionEditorGraphNode* URigMapperDefinitionEditorGraph::ConstructGraphNodesRec(URigMapperDefinition* Definition, const FString& NodeName, bool bIsOutputNode)
 {
 	// No need to do anything is node was already created. Output nodes should not be referenced by other nodes and should only get created once.
 	if (!bIsOutputNode)
@@ -198,22 +199,22 @@ URigMapperDefinitionEditorGraphNode* URigMapperDefinitionEditorGraph::CreateGrap
 	// Node was not created (referenced by a previously created node)
 	if (bIsOutputNode)
 	{
-		return CreateOutputNode(Definition, NodeName);
+		return ConstructOutputNode(Definition, NodeName);
 	}
 	if (Definition->Inputs.Contains(NodeName))
 	{
-		return CreateGraphNode(NodeName, ERigMapperFeatureType::Input);
+		return ConstructGraphNode(NodeName, ERigMapperFeatureType::Input);
 	}
-	return CreateFeatureNode(Definition, NodeName);
+	return ConstructFeatureNode(Definition, NodeName);
 }
 
-URigMapperDefinitionEditorGraphNode* URigMapperDefinitionEditorGraph::CreateOutputNode(URigMapperDefinition* Definition, const FString& NodeName)
+URigMapperDefinitionEditorGraphNode* URigMapperDefinitionEditorGraph::ConstructOutputNode(URigMapperDefinition* Definition, const FString& NodeName)
 {
 	if (const FString* LinkedInputName = Definition->Outputs.Find(NodeName))
 	{
-		URigMapperDefinitionEditorGraphNode* Node = CreateGraphNode(NodeName, ERigMapperFeatureType::Output);
+		URigMapperDefinitionEditorGraphNode* Node = ConstructGraphNode(NodeName, ERigMapperFeatureType::Output);
 		
-		if (URigMapperDefinitionEditorGraphNode* LinkedNode = CreateGraphNodesRec(Definition, *LinkedInputName, false))
+		if (URigMapperDefinitionEditorGraphNode* LinkedNode = ConstructGraphNodesRec(Definition, *LinkedInputName, false))
 		{
 			LinkGraphNodes(LinkedNode, Node);
 		}
@@ -223,25 +224,34 @@ URigMapperDefinitionEditorGraphNode* URigMapperDefinitionEditorGraph::CreateOutp
 	return nullptr;
 }
 
-URigMapperDefinitionEditorGraphNode* URigMapperDefinitionEditorGraph::CreateFeatureNode(URigMapperDefinition* Definition, const FString& NodeName)
+URigMapperDefinitionEditorGraphNode* URigMapperDefinitionEditorGraph::ConstructFeatureNode(URigMapperDefinition* Definition, const FString& NodeName)
 {
 	ERigMapperFeatureType FeatureType;
 
 	if (const FRigMapperFeature* Feature = Definition->Features.Find(NodeName, FeatureType))
 	{
 		
-		URigMapperDefinitionEditorGraphNode* Node = CreateGraphNode(NodeName, FeatureType);
+		URigMapperDefinitionEditorGraphNode* Node = ConstructGraphNode(NodeName, FeatureType);
 			
 		TArray<FString> FeatureInputNames;
 		Feature->GetInputs(FeatureInputNames);
 	
 		for (const FString& FeatureInput : FeatureInputNames)
 		{
-			if (URigMapperDefinitionEditorGraphNode* LinkedNode = CreateGraphNodesRec(Definition, FeatureInput, false))
+			if (URigMapperDefinitionEditorGraphNode* LinkedNode = ConstructGraphNodesRec(Definition, FeatureInput, false))
 			{
 				LinkGraphNodes(LinkedNode, Node);
 			}
 		}
+
+		if (FeatureInputNames.Num() == 0)
+		{
+			Node->CreateInputPin();
+		}
+		if (FeatureInputNames.Num() <= 1 && FeatureType == ERigMapperFeatureType::Multiply)
+		{
+			Node->CreateInputPin();
+		}
 		
 		return Node;
 	}
@@ -269,7 +279,7 @@ void URigMapperDefinitionEditorGraph::LinkGraphNodes(URigMapperDefinitionEditorG
 }
 		
 
-URigMapperDefinitionEditorGraphNode* URigMapperDefinitionEditorGraph::CreateGraphNode(const FString& NodeName, const ERigMapperFeatureType NodeType)
+URigMapperDefinitionEditorGraphNode* URigMapperDefinitionEditorGraph::ConstructGraphNode(const FString& NodeName, const ERigMapperFeatureType NodeType)
 {
 	if (!WeakDefinition.IsValid())
 	{
@@ -283,20 +293,29 @@ URigMapperDefinitionEditorGraphNode* URigMapperDefinitionEditorGraph::CreateGrap
 	GraphNodeCreator.Finalize();
 
 	Node->SetupNode(WeakDefinition.Get(), NodeName, NodeType);
-
+	if (NodeType != ERigMapperFeatureType::NullOutput && NodeType != ERigMapperFeatureType::Output)
+	{
+		// One output node must be created for non output types, even if node is "hanging".
+		UEdGraphPin* OutPin = Node->GetOutputPin();
+		if (!OutPin)
+		{
+			OutPin = Node->CreateOutputPin();
+			OutPin->bHidden = false;
+		}
+	}
 	AddNodeToMap(NodeType, NodeName, Node);
 	
 	return Node;
 }
 
 
-URigMapperDefinitionEditorGraphNode* URigMapperDefinitionEditorGraph::CreateNewGraphNode(const ERigMapperFeatureType InNodeType, UEdGraphPin* InFromPin, const FVector2f& InLocation, bool bSelectNewNode /*= true*/)
+URigMapperDefinitionEditorGraphNode* URigMapperDefinitionEditorGraph::CreateGraphNode(const ERigMapperFeatureType InNodeType, UEdGraphPin* InFromPin, const FVector2f& InLocation, bool bSelectNewNode /*= true*/)
 {
 	if (!WeakDefinition.IsValid())
 	{
 		return nullptr;
 	}
-	FString FromPinName = "New node";
+	FString FromPinName = TEXT("");
 	if (InFromPin)
 	{
 		FromPinName = Cast<URigMapperDefinitionEditorGraphNode>(InFromPin->GetOwningNode())->GetNodeName();
@@ -306,7 +325,8 @@ URigMapperDefinitionEditorGraphNode* URigMapperDefinitionEditorGraph::CreateNewG
 	GraphNodeCreator.Finalize();
 	NewNode->NodePosX = InLocation.X;
 	NewNode->NodePosY = InLocation.Y;
-	FString NodeName = MakeUniqueObjectName(this, URigMapperDefinitionEditorGraphNode::StaticClass(), FName(FromPinName)).ToString();
+	// Create unique feature (node) name.
+	FString NodeName = FRigMapperUtils::GenerateUniqueFeatureName(WeakDefinition, FromPinName, InNodeType);
 	if (InNodeType != ERigMapperFeatureType::Input && InNodeType != ERigMapperFeatureType::NullOutput)
 	{
 		UEdGraphPin* InPin = NewNode->CreateInputPin();
@@ -366,7 +386,6 @@ URigMapperDefinitionEditorGraphNode* URigMapperDefinitionEditorGraph::CreateNewG
 	}
 
 	NewNode->SetupNode(WeakDefinition.Get(), NodeName, InNodeType);
-
 	AddNodeToMap(InNodeType, NodeName, NewNode);
 
 	if (InFromPin)
@@ -378,7 +397,7 @@ URigMapperDefinitionEditorGraphNode* URigMapperDefinitionEditorGraph::CreateNewG
 			const TArray<UEdGraphPin*>& Pins = NewNode->GetInputPins();
 			if (Pins.Num() > 0)
 			{
-				GraphSchema->TryCreateConnection(InFromPin, Pins[0]);
+				InFromPin->MakeLinkTo(Pins[0]);
 			}
 		}
 		else
@@ -386,14 +405,13 @@ URigMapperDefinitionEditorGraphNode* URigMapperDefinitionEditorGraph::CreateNewG
 			UEdGraphPin* OutPin = NewNode->GetOutputPin();
 			if (OutPin)
 			{
-				GraphSchema->TryCreateConnection(OutPin, InFromPin);
-				if (OutPin->LinkedTo.Num() > 0)
-				{
-					Cast<URigMapperDefinitionEditorGraphNode>(OutPin->LinkedTo[0]->GetOwningNode())->AddNewInputLink(NodeName);
-				}
+				OutPin->MakeLinkTo(InFromPin);
+				Cast<URigMapperDefinitionEditorGraphNode>(InFromPin->GetOwningNode())->RelinkFeature(NodeName, InFromPin);
 			}
 		}
 	}
+	// Notify listeners that new Feature is created.
+	WeakDefinition->OnRigMapperDefinitionUpdated.Broadcast();
 
 	return NewNode;
 }
diff --git a/Engine/Plugins/Experimental/Animation/RigMapper/Source/RigMapperEditor/Private/RigMapperGraph/RigMapperDefinitionEditorGraph.h b/Engine/Plugins/Experimental/Animation/RigMapper/Source/RigMapperEditor/Private/RigMapperGraph/RigMapperDefinitionEditorGraph.h
--- a/Engine/Plugins/Experimental/Animation/RigMapper/Source/RigMapperEditor/Private/RigMapperGraph/RigMapperDefinitionEditorGraph.h
+++ b/Engine/Plugins/Experimental/Animation/RigMapper/Source/RigMapperEditor/Private/RigMapperGraph/RigMapperDefinitionEditorGraph.h
@@ -32,8 +32,7 @@ public:
 	
 	URigMapperDefinition* GetDefinition() const { return WeakDefinition.Get(); };
 
-	UE_API URigMapperDefinitionEditorGraphNode* CreateGraphNode(const FString& NodeName, const ERigMapperFeatureType NodeType);
-	UE_API URigMapperDefinitionEditorGraphNode* CreateNewGraphNode(const ERigMapperFeatureType InNodeType, UEdGraphPin* InFromPin, const FVector2f& InLocation, bool bSelectNewNode = true);
+	UE_API URigMapperDefinitionEditorGraphNode* CreateGraphNode(const ERigMapperFeatureType InNodeType, UEdGraphPin* InFromPin, const FVector2f& InLocation, bool bSelectNewNode = true);
 	UE_API TArray<URigMapperDefinitionEditorGraphNode*> GetNodesByName(const TArray<FString>& Inputs, const TArray<FString>& Features, const TArray<FString>& Outputs, const TArray<FString>& NullOutputs) const;
 	/** Retrieves first node that matches the InNodeName (assumes node names are unique	*/
 	UE_API URigMapperDefinitionEditorGraphNode* FindNodeByName(const FString& InNodeName);
@@ -43,13 +42,13 @@ public:
 	UE_API bool RemoveNodeFromMap(const FString& InNodeName, const ERigMapperFeatureType InNodeType);
 
 private:
-	UE_API URigMapperDefinitionEditorGraphNode* CreateGraphNodesRec(URigMapperDefinition* Definition, const FString& InNodeName, bool bIsOutputNode);
-	UE_API URigMapperDefinitionEditorGraphNode* CreateOutputNode(URigMapperDefinition* Definition, const FString& NodeName);
-	UE_API URigMapperDefinitionEditorGraphNode* CreateFeatureNode(URigMapperDefinition* Definition, const FString& NodeName);
+	UE_API URigMapperDefinitionEditorGraphNode* ConstructGraphNode(const FString& NodeName, const ERigMapperFeatureType NodeType);
+	UE_API URigMapperDefinitionEditorGraphNode* ConstructGraphNodesRec(URigMapperDefinition* Definition, const FString& InNodeName, bool bIsOutputNode);
+	UE_API URigMapperDefinitionEditorGraphNode* ConstructOutputNode(URigMapperDefinition* Definition, const FString& NodeName);
+	UE_API URigMapperDefinitionEditorGraphNode* ConstructFeatureNode(URigMapperDefinition* Definition, const FString& NodeName);
 	static UE_API void LinkGraphNodes(URigMapperDefinitionEditorGraphNode* InNode, URigMapperDefinitionEditorGraphNode* OutNode);
 	UE_API void LayoutNodeRec(URigMapperDefinitionEditorGraphNode* InNode, double InputsWidth, double PosY, TArray<URigMapperDefinitionEditorGraphNode*>& LayedOutNodes) const;
 	
-
 	void AddNodeToMap(const ERigMapperFeatureType NodeType, const FString& NodeName, URigMapperDefinitionEditorGraphNode* Node);
 
 private:
diff --git a/Engine/Plugins/Experimental/Animation/RigMapper/Source/RigMapperEditor/Private/RigMapperGraph/RigMapperDefinitionEditorGraphNode.cpp b/Engine/Plugins/Experimental/Animation/RigMapper/Source/RigMapperEditor/Private/RigMapperGraph/RigMapperDefinitionEditorGraphNode.cpp
--- a/Engine/Plugins/Experimental/Animation/RigMapper/Source/RigMapperEditor/Private/RigMapperGraph/RigMapperDefinitionEditorGraphNode.cpp
+++ b/Engine/Plugins/Experimental/Animation/RigMapper/Source/RigMapperEditor/Private/RigMapperGraph/RigMapperDefinitionEditorGraphNode.cpp
@@ -177,7 +177,7 @@ void URigMapperDefinitionEditorGraphNode::UpdateNodeFields()
 	}
 }
 
-void URigMapperDefinitionEditorGraphNode::RelinkFeature(const FString& InOldInput, const FString& InNewInput)
+void URigMapperDefinitionEditorGraphNode::RelinkFeature(const FString& InNewInput, UEdGraphPin* InputPin)
 {
 	if (!Definition.IsValid())
 	{
@@ -188,8 +188,26 @@ void URigMapperDefinitionEditorGraphNode::RelinkFeature(const FString& InOldInpu
 	{
 		return;
 	}
-
-	FeaturePtr->RelinkFeature(InOldInput, InNewInput);
+	if (!InputPin)
+	{
+		return;
+	}
+	TArray<FString> InputNames;
+	FeaturePtr->GetInputs(InputNames);
+	int32 PinIndex = InputPins.Find(InputPin);
+	if (PinIndex == INDEX_NONE)
+	{
+		return;
+	}
+	if (InputNames.Num() <= PinIndex)
+	{
+		// New input should be added to definition.
+		FeaturePtr->AddNewInputLink(InNewInput);
+	}
+	else
+	{
+		FeaturePtr->RelinkFeature(InNewInput, PinIndex);
+	}
 	if (NodeType == ERigMapperFeatureType::Output)
 	{
 		// Definition for Outputs has to be updated manually since RigMapperOutputFeature is created separately.
@@ -223,6 +241,22 @@ void URigMapperDefinitionEditorGraphNode::AddNewInputLink(const FString& InNewIn
 	}
 }
 
+void URigMapperDefinitionEditorGraphNode::RemoveLinksToFeature(const FString& InFeatureName)
+{
+	if (InFeatureName.IsEmpty())
+	{
+		return;
+	}
+
+	FRigMapperFeature* FeaturePtr = GetNodeFeature();
+	if (!FeaturePtr)
+	{
+		return;
+	}
+
+	FeaturePtr->RenameInput(InFeatureName, TEXT(""));
+}
+
 FRigMapperFeature* URigMapperDefinitionEditorGraphNode::GetNodeFeature() const
 {
 	if (Definition.IsValid())
@@ -260,15 +294,41 @@ void URigMapperDefinitionEditorGraphNode::GetRect(FVector2D& TopLeft, FVector2D&
 	BottomRight.Y = NodePosY + GetDimensions().Y + GetMargin().Y;
 }
 
-void URigMapperDefinitionEditorGraphNode::UpdateNodeFromFeature(const FPropertyChangedEvent& InPropertyChangedEvent)
+void URigMapperDefinitionEditorGraphNode::UpdateNodeFromFeature(const FPropertyChangedEvent& InPropertyChangedEvent, TWeakPtr<FStructOnScope> InStructContainer)
 {
-	FRigMapperFeature* FeaturePtr = GetNodeFeature();
+	if (!InStructContainer.IsValid())
+	{
+		return;
+	}
+	if (!Definition.IsValid())
+	{
+		return;
+	}
+	FString OldName = NodeName;
+	FRigMapperFeature* FeaturePtr = nullptr;
+	if ((InPropertyChangedEvent.ChangeType & EPropertyChangeType::ValueSet) && InPropertyChangedEvent.GetPropertyName() == GET_MEMBER_NAME_CHECKED(FRigMapperFeature, Name))
+	{
+		// Name has been modified, find feature using new name.
+		if (const FStrProperty* NameProp = CastField<FStrProperty>(InPropertyChangedEvent.Property))
+		{
+			void* StructMem = InStructContainer.Pin()->GetStructMemory();
+			if (const FString* NamePtr = NameProp->ContainerPtrToValuePtr<FString>(StructMem))
+			{
+				FeaturePtr = Definition->GetFeature(*NamePtr, NodeType);
+			}
+		}
+	}
+	else
+	{
+		FeaturePtr = GetNodeFeature();
+	}	
 	if (!FeaturePtr)
 	{
 		return;
 	}
 
 	Modify();
+	NodeName = FeaturePtr->Name;
 
 	const FName PropertyName = InPropertyChangedEvent.GetPropertyName();
 	const FName ParentName = InPropertyChangedEvent.GetMemberPropertyName();
@@ -349,7 +409,7 @@ void URigMapperDefinitionEditorGraphNode::UpdateNodeFromFeature(const FPropertyC
 			// CHANGING AN INPUT NAME
 			if (PinIndex != INDEX_NONE)
 			{
-				URigMapperDefinitionEditorGraph* RigMapperGraph = static_cast<URigMapperDefinitionEditorGraph*>(GetGraph());
+				URigMapperDefinitionEditorGraph* RigMapperGraph = Cast<URigMapperDefinitionEditorGraph>(GetGraph());
 				URigMapperDefinitionEditorGraphNode* LinkedNode = nullptr;
 				if (NodeType == ERigMapperFeatureType::Multiply)
 				{
@@ -378,9 +438,7 @@ void URigMapperDefinitionEditorGraphNode::UpdateNodeFromFeature(const FPropertyC
 		}
 	}
 
-	FString NewName = FeaturePtr->Name;	
-
-	if (NewName != NodeName && NodeType != ERigMapperFeatureType::Output && NodeType != ERigMapperFeatureType::NullOutput)
+	if (NodeName != OldName && NodeType != ERigMapperFeatureType::Output && NodeType != ERigMapperFeatureType::NullOutput) // Output and NullOutput nodes don't have output pins.
 	{
 		// Update inputs of all nodes directly connected to this node.
 		for (UEdGraphPin* OutPin : OutputPins) // It will always be one output pin.
@@ -389,13 +447,15 @@ void URigMapperDefinitionEditorGraphNode::UpdateNodeFromFeature(const FPropertyC
 			{
 				if (URigMapperDefinitionEditorGraphNode* LinkedNode = Cast<URigMapperDefinitionEditorGraphNode>(LinkedPin->GetOwningNode()))
 				{
-					LinkedNode->RelinkFeature(NodeName, NewName);
+					LinkedNode->RelinkFeature(NodeName, LinkedPin);
 				}
 			}
 		}
+		// Update node map.
+		URigMapperDefinitionEditorGraph* RigMapperGraph = Cast<URigMapperDefinitionEditorGraph>(GetGraph());
+		RigMapperGraph->UpdateNodeNameInMap(OldName, NodeName, NodeType);
 	}
-	
-	NodeName = NewName;
+		
 	NodeTitle = FText::FromString(NodeName);
 	UpdateNodeFields();	
 }
diff --git a/Engine/Plugins/Experimental/Animation/RigMapper/Source/RigMapperEditor/Private/RigMapperGraph/RigMapperDefinitionEditorGraphNode.h b/Engine/Plugins/Experimental/Animation/RigMapper/Source/RigMapperEditor/Private/RigMapperGraph/RigMapperDefinitionEditorGraphNode.h
--- a/Engine/Plugins/Experimental/Animation/RigMapper/Source/RigMapperEditor/Private/RigMapperGraph/RigMapperDefinitionEditorGraphNode.h
+++ b/Engine/Plugins/Experimental/Animation/RigMapper/Source/RigMapperEditor/Private/RigMapperGraph/RigMapperDefinitionEditorGraphNode.h
@@ -12,6 +12,7 @@
 
 enum class ERigMapperFeatureType : uint8;
 class URigMapperDefinition;
+class FStructOnScope;
 struct FRigMapperFeature;
 
 /**
@@ -61,13 +62,14 @@ public:
 	UE_API void GetRect(FVector2D& TopLeft, FVector2D& BottomRight) const;
 
 	/** Updates node from definition and makes it dirty */
-	UE_API void UpdateNodeFromFeature(const FPropertyChangedEvent& InPropertyChangedEvent);
+	UE_API void UpdateNodeFromFeature(const FPropertyChangedEvent& InPropertyChangedEvent, TWeakPtr<FStructOnScope> InStructContainer);
 
 	/** Updates node body */
 	UE_API void UpdateNodeFields();
 
-	UE_API void RelinkFeature(const FString& InOldInput, const FString& InNewInput);
+	UE_API void RelinkFeature(const FString& InNewInput, UEdGraphPin* InputPin);
 	UE_API void AddNewInputLink(const FString& InNewInput);
+	UE_API void RemoveLinksToFeature(const FString& InFeatureName);
 
 protected:
 	/** Cached title for the node */
diff --git a/Engine/Plugins/Experimental/Animation/RigMapper/Source/RigMapperEditor/Private/RigMapperGraph/RigMapperDefinitionEditorGraphSchema.cpp b/Engine/Plugins/Experimental/Animation/RigMapper/Source/RigMapperEditor/Private/RigMapperGraph/RigMapperDefinitionEditorGraphSchema.cpp
--- a/Engine/Plugins/Experimental/Animation/RigMapper/Source/RigMapperEditor/Private/RigMapperGraph/RigMapperDefinitionEditorGraphSchema.cpp
+++ b/Engine/Plugins/Experimental/Animation/RigMapper/Source/RigMapperEditor/Private/RigMapperGraph/RigMapperDefinitionEditorGraphSchema.cpp
@@ -36,7 +36,7 @@
 		TSharedPtr<FScopedTransaction> Transaction = MakeShared<FScopedTransaction>(LOCTEXT("RigMapperDefinitionEditor_SpawnNewNode", "Spawn RigMapper definition node"));
 		Graph->Modify();
 
-		NewNode = Graph->CreateNewGraphNode(NodeType, FromPin, Location, bSelectNewNode);
+		NewNode = Graph->CreateGraphNode(NodeType, FromPin, Location, bSelectNewNode);
 		
 		Transaction.Reset();
 	}
@@ -50,14 +50,14 @@
 		// TODO: Implement context menu actions.
 		FToolMenuSection& Section = Menu->AddSection("RigMapperDefinitionEditorNodeActions", LOCTEXT("RigMapperDefinitionEditor_NodeActionsMenuHeader", "Node Actions"));
 		Section.AddMenuEntry(FGenericCommands::Get().Delete);
-		Section.AddMenuEntry(FGenericCommands::Get().Cut);
-		Section.AddMenuEntry(FGenericCommands::Get().Copy);
-		Section.AddMenuEntry(FGenericCommands::Get().Duplicate);
-		Section.AddMenuEntry(FGenericCommands::Get().Rename);
+		//Section.AddMenuEntry(FGenericCommands::Get().Cut);
+		//Section.AddMenuEntry(FGenericCommands::Get().Copy);
+		//Section.AddMenuEntry(FGenericCommands::Get().Duplicate);
+		//Section.AddMenuEntry(FGenericCommands::Get().Rename);
 		Section.AddMenuEntry(FGraphEditorCommands::Get().BreakNodeLinks);
-		Section.AddMenuEntry(FGraphEditorCommands::Get().FindReferences, LOCTEXT("RigMapperDefinitionEditor_FindReferences", "Find References In Graph"),
-			LOCTEXT("RigMapperDefinitionEditor_FindReferencesTooltip", "Find References to the selected Node in the current Graph"),
-			FSlateIcon());
+		//Section.AddMenuEntry(FGraphEditorCommands::Get().FindReferences, LOCTEXT("RigMapperDefinitionEditor_FindReferences", "Find References In Graph"),
+		//	LOCTEXT("RigMapperDefinitionEditor_FindReferencesTooltip", "Find References to the selected Node in the current Graph"),
+		//	FSlateIcon());
 
 	}
 	Super::GetContextMenuActions(Menu, Context);	
@@ -94,53 +94,6 @@
 	return new ConnectionDrawingPolicy(InBackLayerID, InFrontLayerID, InZoomFactor, InClippingRect, InDrawElements);
 }
 
-const FPinConnectionResponse URigMapperDefinitionEditorGraphSchema::CanCreateConnection(const UEdGraphPin* InA, const UEdGraphPin* InB) const
-{
-	// Make sure the pins are not on the same node
-	if (InA->GetOwningNode() == InB->GetOwningNode())
-	{
-		return FPinConnectionResponse(CONNECT_RESPONSE_DISALLOW, LOCTEXT("RigMapperDefinitionEditor_ConnectionSameNode", "Both are on the same node"));
-	}
-
-	// Compare the directions
-	const UEdGraphPin* InputPin = NULL;
-	const UEdGraphPin* OutputPin = NULL;
-
-	if (!CategorizePinsByDirection(InA, InB, /*out*/ InputPin, /*out*/ OutputPin))
-	{
-		return FPinConnectionResponse(CONNECT_RESPONSE_DISALLOW, LOCTEXT("RigMapperDefinitionEditor_ConnectionIncompatible", "Directions are not compatible"));
-	}
-
-	// Check for new and existing loops
-	FText ResponseMessage;
-	//if (ConnectionCausesLoop(InputPin, OutputPin))
-	//{
-	//	ResponseMessage = LOCTEXT("ConnectionLoop", "Connection could cause loop");
-	//}
-
-	// For non-exec pins, break existing connections on inputs only - multiple output connections are acceptable
-	if (InputPin->LinkedTo.Num() > 0)
-	{	
-		ECanCreateConnectionResponse ReplyBreakOutputs;
-		if (InputPin == InA)
-		{
-			ReplyBreakOutputs = CONNECT_RESPONSE_BREAK_OTHERS_A;
-		}
-		else
-		{
-			ReplyBreakOutputs = CONNECT_RESPONSE_BREAK_OTHERS_B;
-		}
-		if (ResponseMessage.IsEmpty())
-		{
-			ResponseMessage = LOCTEXT("RigMapperDefinitionEditor_ConnectionReplace", "Replace existing connections");
-		}
-		return FPinConnectionResponse(ReplyBreakOutputs, ResponseMessage);
-
-	}
-
-	return FPinConnectionResponse(CONNECT_RESPONSE_MAKE, TEXT(""));
-}
-
 FPinConnectionResponse URigMapperDefinitionEditorGraphSchema::CanCreateNewNodes(UEdGraphPin* InSourcePin) const
 {
 	return FPinConnectionResponse(CONNECT_RESPONSE_MAKE, TEXT(""));
@@ -217,6 +170,52 @@
 	}
 }
 
+const FPinConnectionResponse URigMapperDefinitionEditorGraphSchema::CanCreateConnection(const UEdGraphPin* InA, const UEdGraphPin* InB) const
+{
+	// Make sure the pins are not on the same node
+	if (InA->GetOwningNode() == InB->GetOwningNode())
+	{
+		return FPinConnectionResponse(CONNECT_RESPONSE_DISALLOW, LOCTEXT("RigMapperDefinitionEditor_ConnectionSameNode", "Both are on the same node"));
+	}
+
+	// Compare the directions
+	const UEdGraphPin* InputPin = NULL;
+	const UEdGraphPin* OutputPin = NULL;
+
+	if (!CategorizePinsByDirection(InA, InB, /*out*/ InputPin, /*out*/ OutputPin))
+	{
+		return FPinConnectionResponse(CONNECT_RESPONSE_DISALLOW, LOCTEXT("RigMapperDefinitionEditor_ConnectionIncompatible", "Directions are not compatible"));
+	}
+
+	// Check for new and existing loops
+	FText ResponseMessage;
+	//if (ConnectionCausesLoop(InputPin, OutputPin))
+	//{
+	//	ResponseMessage = LOCTEXT("ConnectionLoop", "Connection could cause loop");
+	//}
+
+	// For non-exec pins, break existing connections on inputs only - multiple output connections are acceptable
+	if (InputPin->LinkedTo.Num() > 0)
+	{
+		ECanCreateConnectionResponse ReplyBreakOutputs;
+		if (InputPin == InA)
+		{
+			ReplyBreakOutputs = CONNECT_RESPONSE_BREAK_OTHERS_A;
+		}
+		else
+		{
+			ReplyBreakOutputs = CONNECT_RESPONSE_BREAK_OTHERS_B;
+		}
+		if (ResponseMessage.IsEmpty())
+		{
+			ResponseMessage = LOCTEXT("RigMapperDefinitionEditor_ConnectionReplace", "Replace existing connections");
+		}
+		return FPinConnectionResponse(ReplyBreakOutputs, ResponseMessage);
+	}
+
+	return FPinConnectionResponse(CONNECT_RESPONSE_MAKE, TEXT(""));
+}
+
 bool URigMapperDefinitionEditorGraphSchema::TryCreateConnection(UEdGraphPin* InA, UEdGraphPin* InB) const
 {
 	if (InA == InB)
@@ -228,7 +227,57 @@
 		return false;
 	}
 
-	return Super::TryCreateConnection(InA, InB);
+	bool bModified = false;
+	FPinConnectionResponse ConnResponse = CanCreateConnection(InA, InB);
+	if (ConnResponse.Response == CONNECT_RESPONSE_DISALLOW)
+	{
+		return false;
+	}
+
+	// If there is an input pin with an existing connection, it is about to be broken.
+	// We need to save the name of that node in order to update the Definition.
+	UEdGraphPin* InputPin = nullptr;
+	const UEdGraphPin* OutputPin = nullptr;
+	FString OldInputName = TEXT("");
+	FString NewInputName = TEXT("");
+	if (ConnResponse.Response == CONNECT_RESPONSE_BREAK_OTHERS_A)
+	{
+		// Pin A is an input pin with existing connection which is about to be broken.
+		InputPin = InA;
+		OutputPin = InB;
+	}
+	else if (ConnResponse.Response == CONNECT_RESPONSE_BREAK_OTHERS_B)
+	{
+		// Pin B is an input pin with existing connection which is about to be broken.
+		InputPin = InB;
+		OutputPin = InA;
+	}
+	else if (InA->Direction == EGPD_Input)
+	{
+		InputPin = InA;
+		OutputPin = InB;
+	}
+	else
+	{
+		InputPin = InB;
+		OutputPin = InA;
+	}
+	if (OutputPin)
+	{
+		if (URigMapperDefinitionEditorGraphNode* NewNode = Cast<URigMapperDefinitionEditorGraphNode>(OutputPin->GetOwningNode()))
+		{
+			NewInputName = NewNode->GetNodeName();
+		}
+	}
+	bModified = Super::TryCreateConnection(InA, InB);
+	
+	if (bModified)
+	{
+		URigMapperDefinitionEditorGraphNode* InputNode = Cast<URigMapperDefinitionEditorGraphNode>(InputPin->GetOwningNode());
+		InputNode->RelinkFeature(NewInputName, InputPin);
+	}
+
+	return bModified;
 }
 
 bool URigMapperDefinitionEditorGraphSchema::IsConnectionRelinkingAllowed(UEdGraphPin* InPin) const
@@ -255,4 +304,40 @@
 	return Super::TryRelinkConnectionTarget(SourcePin, OldTargetPin, NewTargetPin, InSelectedGraphNodes);
 }
 
+bool URigMapperDefinitionEditorGraphSchema::SafeDeleteNodeFromGraph(UEdGraph* Graph, UEdGraphNode* Node) const
+{
+	if (!Graph)
+	{
+		return false;
+	}
+
+	if (!Node)
+	{
+		return false;
+	}
+
+	if (URigMapperDefinitionEditorGraphNode* RigMapperNode = Cast<URigMapperDefinitionEditorGraphNode>(Node))
+	{
+		if (const UEdGraphPin* OutputPin = RigMapperNode->GetOutputPin())
+		{
+			for (const UEdGraphPin* LinkedPin : OutputPin->LinkedTo)
+			{
+				if (URigMapperDefinitionEditorGraphNode* LinkedNode = Cast<URigMapperDefinitionEditorGraphNode>(LinkedPin->GetOwningNode()))
+				{
+					LinkedNode->RemoveLinksToFeature(RigMapperNode->GetNodeName());
+				}
+			}
+		}
+
+		if (URigMapperDefinitionEditorGraph* RigMapperGraph = Cast<URigMapperDefinitionEditorGraph>(Graph))
+		{
+			RigMapperGraph->RemoveNodeFromMap(RigMapperNode->GetNodeName(), RigMapperNode->GetNodeType());
+		}
+
+
+		return true;		
+	}
+	return false;
+}
+
 #undef LOCTEXT_NAMESPACE
diff --git a/Engine/Plugins/Experimental/Animation/RigMapper/Source/RigMapperEditor/Private/RigMapperGraph/RigMapperDefinitionEditorGraphSchema.h b/Engine/Plugins/Experimental/Animation/RigMapper/Source/RigMapperEditor/Private/RigMapperGraph/RigMapperDefinitionEditorGraphSchema.h
--- a/Engine/Plugins/Experimental/Animation/RigMapper/Source/RigMapperEditor/Private/RigMapperGraph/RigMapperDefinitionEditorGraphSchema.h
+++ b/Engine/Plugins/Experimental/Animation/RigMapper/Source/RigMapperEditor/Private/RigMapperGraph/RigMapperDefinitionEditorGraphSchema.h
@@ -66,9 +66,9 @@ public:
 	UE_API virtual class FConnectionDrawingPolicy* CreateConnectionDrawingPolicy(int32 InBackLayerID, int32 InFrontLayerID, float InZoomFactor, const FSlateRect& InClippingRect, class FSlateWindowElementList& InDrawElements, class UEdGraph* InGraphObj) const override;
 	virtual bool ShouldAlwaysPurgeOnModification() const override { return false; }
 	UE_API virtual FPinConnectionResponse CanCreateNewNodes(UEdGraphPin* InSourcePin) const override;
-	UE_API virtual const FPinConnectionResponse CanCreateConnection(const UEdGraphPin* InA, const UEdGraphPin* InB) const override;
 	UE_API virtual bool SupportsDropPinOnNode(UEdGraphNode* InTargetNode, const FEdGraphPinType& InSourcePinType, EEdGraphPinDirection InSourcePinDirection, FText& OutErrorMessage) const override;
 	UE_API void GetGraphContextActions(FGraphContextMenuBuilder& ContextMenuBuilder) const override;
+	UE_API virtual const FPinConnectionResponse CanCreateConnection(const UEdGraphPin* InA, const UEdGraphPin* InB) const override;
 	UE_API bool TryCreateConnection(UEdGraphPin* InA, UEdGraphPin* InB) const override;
 	UE_API bool IsConnectionRelinkingAllowed(UEdGraphPin* InPin) const override;
 	UE_API const FPinConnectionResponse CanRelinkConnectionToPin(const UEdGraphPin* OldSourcePin, const UEdGraphPin* TargetPinCandidate) const override;
@@ -76,11 +76,8 @@ public:
 	UEdGraphPin* DropPinOnNode(UEdGraphNode* InTargetNode, const FName& InSourcePinName, const FEdGraphPinType& InSourcePinType, EEdGraphPinDirection InSourcePinDirection) const override;
 	void SetPinBeingDroppedOnNode(UEdGraphPin* InSourcePin) const override;
 	UE_API bool TryRelinkConnectionTarget(UEdGraphPin* SourcePin, UEdGraphPin* OldTargetPin, UEdGraphPin* NewTargetPin, const TArray<UEdGraphNode*>& InSelectedGraphNodes) const override;
+	bool SafeDeleteNodeFromGraph(UEdGraph* Graph, UEdGraphNode* Node) const override;
 	// End of UEdGraphSchema interface
-private:
-
-
-
 };
 
 #undef UE_API
diff --git a/Engine/Plugins/Experimental/Animation/RigMapper/Source/RigMapperEditor/Private/RigMapperGraph/SRigMapperDefinitionGraphEditor.cpp b/Engine/Plugins/Experimental/Animation/RigMapper/Source/RigMapperEditor/Private/RigMapperGraph/SRigMapperDefinitionGraphEditor.cpp
--- a/Engine/Plugins/Experimental/Animation/RigMapper/Source/RigMapperEditor/Private/RigMapperGraph/SRigMapperDefinitionGraphEditor.cpp
+++ b/Engine/Plugins/Experimental/Animation/RigMapper/Source/RigMapperEditor/Private/RigMapperGraph/SRigMapperDefinitionGraphEditor.cpp
@@ -5,6 +5,9 @@
 #include "RigMapperDefinitionEditorGraph.h"
 #include "RigMapperDefinitionEditorGraphNode.h"
 #include "RigMapperDefinitionEditorGraphSchema.h"
+#include "GraphEditorActions.h"
+#include "ScopedTransaction.h"
+#include "Framework/Commands/GenericCommands.h"
 
 #include "SGraphPanel.h"
 #include "SlateOptMacros.h"
@@ -40,6 +43,7 @@ void SRigMapperDefinitionGraphEditor::Construct(const FArguments& InArgs, URigMa
 	GraphObj->Initialize(InDefinition);
 	GraphObj->RebuildGraph();
 	GraphObj->LayoutNodes();
+	BindCommands();
 
 	ChildSlot
 	[
@@ -54,6 +58,7 @@ void SRigMapperDefinitionGraphEditor::Construct(const FArguments& InArgs, URigMa
 			.ShowGraphStateOverlay(false)
 			.IsEditable(true)
 			.AutoExpandActionMenu(true)
+			.AdditionalCommands(UICommandList)
 			.DisplayAsReadOnly(false)				
 		]
 	];
@@ -230,6 +235,67 @@ void SRigMapperDefinitionGraphEditor::HandleSelectionChanged(const TSet<UObject*
 	}
 }
 
+FReply SRigMapperDefinitionGraphEditor::OnKeyDown(const FGeometry& MyGeometry, const FKeyEvent& InKeyEvent)
+{
+	return UICommandList->ProcessCommandBindings(InKeyEvent) ? FReply::Handled() : FReply::Unhandled();
+}
+
+void SRigMapperDefinitionGraphEditor::BindCommands()
+{
+	// This should not be called twice on the same instance
+	check(!UICommandList.IsValid());
+
+	UICommandList = MakeShareable(new FUICommandList);
+
+	FUICommandList& CommandList = *UICommandList;
+
+	// ...and bind them all
+	CommandList.MapAction(
+		FGenericCommands::Get().Delete,
+		FExecuteAction::CreateSP(this, &SRigMapperDefinitionGraphEditor::OnDeleteSelectedNodes),
+		FCanExecuteAction::CreateSP(this, &SRigMapperDefinitionGraphEditor::CanDeleteNodes));
+}
+
+void SRigMapperDefinitionGraphEditor::OnDeleteSelectedNodes()
+{
+	UEdGraph* CurrentGraph = GraphEditor->GetCurrentGraph();
+	if (!CurrentGraph)
+	{
+		return;
+	}
+
+	const UEdGraphSchema* Schema = CurrentGraph->GetSchema();	
+	const FGraphPanelSelectionSet& Selection = GraphEditor->GetSelectedNodes();
+	TSharedPtr<FScopedTransaction> Transaction = MakeShared<FScopedTransaction>(FGenericCommands::Get().Delete->GetDescription());
+
+	TArray<URigMapperDefinitionEditorGraphNode*> NodesToDelete;
+	NodesToDelete.Reserve(Selection.Num());
+
+	for (UObject* Obj : Selection)
+	{
+		if (URigMapperDefinitionEditorGraphNode* Node = Cast<URigMapperDefinitionEditorGraphNode>(Obj))
+		{
+			NodesToDelete.Add(Node);
+		}
+	}
+
+	CurrentGraph->Modify();
+	for (UEdGraphNode* Node : NodesToDelete)
+	{
+		if (Schema)
+		{
+			Schema->SafeDeleteNodeFromGraph(CurrentGraph, Node);
+		}
+		Node->Modify();
+		Node->DestroyNode();
+	}
+	GraphEditor->ClearSelectionSet();
+}
+
+bool SRigMapperDefinitionGraphEditor::CanDeleteNodes()
+{
+	return true;
+}
 
 #undef LOCTEXT_NAMESPACE
 
diff --git a/Engine/Plugins/Experimental/Animation/RigMapper/Source/RigMapperEditor/Private/RigMapperGraph/SRigMapperDefinitionGraphEditor.h b/Engine/Plugins/Experimental/Animation/RigMapper/Source/RigMapperEditor/Private/RigMapperGraph/SRigMapperDefinitionGraphEditor.h
--- a/Engine/Plugins/Experimental/Animation/RigMapper/Source/RigMapperEditor/Private/RigMapperGraph/SRigMapperDefinitionGraphEditor.h
+++ b/Engine/Plugins/Experimental/Animation/RigMapper/Source/RigMapperEditor/Private/RigMapperGraph/SRigMapperDefinitionGraphEditor.h
@@ -39,13 +39,20 @@ private:
 	UE_API void ZoomToFitNodes(const TArray<URigMapperDefinitionEditorGraphNode*>& SelectedNodes) const;
 	static UE_API void GetAllLinkedNodes(const URigMapperDefinitionEditorGraphNode* BaseNode, TArray<URigMapperDefinitionEditorGraphNode*>& LinkedNodes, bool bDescend);
 	UE_API void HandleSelectionChanged(const TSet<UObject*>& Nodes);
+	virtual FReply OnKeyDown(const FGeometry& MyGeometry, const FKeyEvent& InKeyEvent) override;
+
+	void BindCommands();
+	void OnDeleteSelectedNodes();
+	bool CanDeleteNodes();
 
 public:
-	SGraphEditor::FOnSelectionChanged OnSelectionChanged;
+	SGraphEditor::FOnSelectionChanged OnSelectionChanged;	
 	
 private:
 	TSharedPtr<SGraphEditor> GraphEditor;
 	TObjectPtr<URigMapperDefinitionEditorGraph> GraphObj;
+	/** Commands that are bound to delegates*/
+	TSharedPtr<FUICommandList> UICommandList;
 
 	bool bSelectingNodes = false;
 
