[Rig Mapper] Deleting nodes now updates Definition and SideView

diff --git a/Engine/Plugins/Experimental/Animation/RigMapper/Source/RigMapper/Private/RigMapperDefinition.cpp b/Engine/Plugins/Experimental/Animation/RigMapper/Source/RigMapper/Private/RigMapperDefinition.cpp
--- a/Engine/Plugins/Experimental/Animation/RigMapper/Source/RigMapper/Private/RigMapperDefinition.cpp
+++ b/Engine/Plugins/Experimental/Animation/RigMapper/Source/RigMapper/Private/RigMapperDefinition.cpp
@@ -723,6 +723,43 @@ FRigMapperFeature* URigMapperDefinition::GetFeature(const FString& InFeatureName
 	return nullptr;
 }
 
+bool URigMapperDefinition::DeleteFeature(const FString& InFeatureName, const ERigMapperFeatureType InFeatureType)
+{
+	switch (InFeatureType)
+	{
+	case ERigMapperFeatureType::Multiply:
+		return Features.Multiply.RemoveAll([InFeatureName](const FRigMapperFeature& InOther) { return InOther.Name == InFeatureName; }) != 0;		
+
+	case ERigMapperFeatureType::WeightedSum:
+		return Features.WeightedSums.RemoveAll([InFeatureName](const FRigMapperWsFeature& Other) { return Other.Name == InFeatureName; }) != 0;
+
+	case ERigMapperFeatureType::SDK:
+		return Features.SDKs.RemoveAll([InFeatureName](const FRigMapperSdkFeature& Other) { return Other.Name == InFeatureName; }) != 0;
+
+	case ERigMapperFeatureType::Input:	
+		if (Features.Inputs.RemoveAll([InFeatureName](const FRigMapperFeature& Other) { return Other.Name == InFeatureName; }) != 0)
+		{
+			return Inputs.Remove(InFeatureName) != 0;
+		}
+		break;
+	case ERigMapperFeatureType::Output:		
+		if (Features.Outputs.RemoveAll([InFeatureName](const FRigMapperOutputFeature& Other) { return Other.Name == InFeatureName; }) != 0)
+		{
+			return Outputs.Remove(InFeatureName) != 0;
+		}
+		break;
+	case ERigMapperFeatureType::NullOutput:
+		if (Features.NullOutputs.RemoveAll([InFeatureName](const FRigMapperNullOutputFeature& Other) { return Other.Name == InFeatureName; }) != 0)
+		{
+			return NullOutputs.Remove(InFeatureName) != 0;
+		}
+		break;
+	default:
+		break;
+	}
+	return false;
+}
+
 bool URigMapperDefinition::IsDefinitionValid(bool bWarn, bool bForce) const
 {
 	if (!bForce && WasDefinitionValidated())
@@ -982,7 +1019,14 @@ void FRigMapperWsFeature::RelinkFeature(const FString& InNewInput, const int32 I
 		{
 			double TempValue = MoveTemp(*Weight);
 			Inputs.Remove(OldInput);
-			Inputs.Add(InNewInput, TempValue);
+			if (!Inputs.Contains(InNewInput))
+			{
+				Inputs.Add(InNewInput, TempValue);
+			}
+			else
+			{
+				Inputs[InNewInput] += TempValue;
+			}
 			OrderedInputs[InputIndex].Name = InNewInput;
 		}
 	}
diff --git a/Engine/Plugins/Experimental/Animation/RigMapper/Source/RigMapper/Private/RigMapperUtils.cpp b/Engine/Plugins/Experimental/Animation/RigMapper/Source/RigMapper/Private/RigMapperUtils.cpp
--- a/Engine/Plugins/Experimental/Animation/RigMapper/Source/RigMapper/Private/RigMapperUtils.cpp
+++ b/Engine/Plugins/Experimental/Animation/RigMapper/Source/RigMapper/Private/RigMapperUtils.cpp
@@ -67,10 +67,6 @@ FString FRigMapperUtils::GenerateUniqueFeatureName(const TWeakObjectPtr<URigMapp
 				return Prefix + InDesiredName;
 			}();
 	}
-	if (!InDesiredName.IsEmpty() && !ExistingNames.Contains(InDesiredName))
-	{
-		return InDesiredName;
-	}
 
 	if (InDesiredName.IsEmpty())
 	{
diff --git a/Engine/Plugins/Experimental/Animation/RigMapper/Source/RigMapper/Public/RigMapperDefinition.h b/Engine/Plugins/Experimental/Animation/RigMapper/Source/RigMapper/Public/RigMapperDefinition.h
--- a/Engine/Plugins/Experimental/Animation/RigMapper/Source/RigMapper/Public/RigMapperDefinition.h
+++ b/Engine/Plugins/Experimental/Animation/RigMapper/Source/RigMapper/Public/RigMapperDefinition.h
@@ -437,6 +437,8 @@ public:
 	 */
 	UE_API FRigMapperFeature* GetFeature(const FString& InFeatureName, const ERigMapperFeatureType InFeatureType);
 
+	UE_API bool DeleteFeature(const FString& InFeatureName, const ERigMapperFeatureType InFeatureType);
+
 private:
 	bool LoadInputsFromJsonObject(const TSharedPtr<FJsonObject>& JsonObject);
 	bool LoadFeaturesFromJsonObject(const TSharedPtr<FJsonObject>& JsonObject, TArray<FString>& FeatureNames);
diff --git a/Engine/Plugins/Experimental/Animation/RigMapper/Source/RigMapperEditor/Private/RigMapperGraph/RigMapperDefinitionEditorGraph.cpp b/Engine/Plugins/Experimental/Animation/RigMapper/Source/RigMapperEditor/Private/RigMapperGraph/RigMapperDefinitionEditorGraph.cpp
--- a/Engine/Plugins/Experimental/Animation/RigMapper/Source/RigMapperEditor/Private/RigMapperGraph/RigMapperDefinitionEditorGraph.cpp
+++ b/Engine/Plugins/Experimental/Animation/RigMapper/Source/RigMapperEditor/Private/RigMapperGraph/RigMapperDefinitionEditorGraph.cpp
@@ -158,6 +158,7 @@ bool URigMapperDefinitionEditorGraph::UpdateNodeNameInMap(const FString& InOldNa
 
 bool URigMapperDefinitionEditorGraph::RemoveNodeFromMap(const FString& InNodeName, const ERigMapperFeatureType InNodeType)
 {
+	WeakDefinition->DeleteFeature(InNodeName, InNodeType);
 	switch (InNodeType)
 	{
 	case ERigMapperFeatureType::WeightedSum:
diff --git a/Engine/Plugins/Experimental/Animation/RigMapper/Source/RigMapperEditor/Private/RigMapperGraph/SRigMapperDefinitionGraphEditor.cpp b/Engine/Plugins/Experimental/Animation/RigMapper/Source/RigMapperEditor/Private/RigMapperGraph/SRigMapperDefinitionGraphEditor.cpp
--- a/Engine/Plugins/Experimental/Animation/RigMapper/Source/RigMapperEditor/Private/RigMapperGraph/SRigMapperDefinitionGraphEditor.cpp
+++ b/Engine/Plugins/Experimental/Animation/RigMapper/Source/RigMapperEditor/Private/RigMapperGraph/SRigMapperDefinitionGraphEditor.cpp
@@ -5,6 +5,7 @@
 #include "RigMapperDefinitionEditorGraph.h"
 #include "RigMapperDefinitionEditorGraphNode.h"
 #include "RigMapperDefinitionEditorGraphSchema.h"
+#include "RigMapperDefinition.h"
 #include "GraphEditorActions.h"
 #include "ScopedTransaction.h"
 #include "Framework/Commands/GenericCommands.h"
@@ -290,6 +291,10 @@ void SRigMapperDefinitionGraphEditor::OnDeleteSelectedNodes()
 		Node->DestroyNode();
 	}
 	GraphEditor->ClearSelectionSet();
+	if (URigMapperDefinitionEditorGraph* RigMapperGraph = Cast<URigMapperDefinitionEditorGraph>(CurrentGraph))
+	{
+		RigMapperGraph->GetDefinition()->OnRigMapperDefinitionUpdated.Broadcast();
+	}
 }
 
 bool SRigMapperDefinitionGraphEditor::CanDeleteNodes()
