[RigMapper] Implement RigMapperDefinitionGraph pin interactions:
- Breaking links
- Moving links
- Adding new connection by dragging and dropping

diff --git a/Engine/Plugins/Experimental/Animation/RigMapper/Source/RigMapper/Private/RigMapperDefinition.cpp b/Engine/Plugins/Experimental/Animation/RigMapper/Source/RigMapper/Private/RigMapperDefinition.cpp
--- a/Engine/Plugins/Experimental/Animation/RigMapper/Source/RigMapper/Private/RigMapperDefinition.cpp
+++ b/Engine/Plugins/Experimental/Animation/RigMapper/Source/RigMapper/Private/RigMapperDefinition.cpp
@@ -736,22 +736,25 @@ bool URigMapperDefinition::DeleteFeature(const FString& InFeatureName, const ERi
 	case ERigMapperFeatureType::SDK:
 		return Features.SDKs.RemoveAll([InFeatureName](const FRigMapperSdkFeature& Other) { return Other.Name == InFeatureName; }) != 0;
 
-	case ERigMapperFeatureType::Input:	
+	case ERigMapperFeatureType::Input:		
 		if (Features.Inputs.RemoveAll([InFeatureName](const FRigMapperFeature& Other) { return Other.Name == InFeatureName; }) != 0)
 		{
-			return Inputs.Remove(InFeatureName) != 0;
+			Inputs.Remove(InFeatureName);
+			return true;
 		}
 		break;
 	case ERigMapperFeatureType::Output:		
 		if (Features.Outputs.RemoveAll([InFeatureName](const FRigMapperOutputFeature& Other) { return Other.Name == InFeatureName; }) != 0)
 		{
-			return Outputs.Remove(InFeatureName) != 0;
+			Outputs.Remove(InFeatureName);
+			return true;
 		}
 		break;
 	case ERigMapperFeatureType::NullOutput:
 		if (Features.NullOutputs.RemoveAll([InFeatureName](const FRigMapperNullOutputFeature& Other) { return Other.Name == InFeatureName; }) != 0)
 		{
-			return NullOutputs.Remove(InFeatureName) != 0;
+			NullOutputs.Remove(InFeatureName);
+			return true;
 		}
 		break;
 	default:
diff --git a/Engine/Plugins/Experimental/Animation/RigMapper/Source/RigMapperEditor/Private/RigMapperGraph/RigMapperDefinitionEditorGraphNode.cpp b/Engine/Plugins/Experimental/Animation/RigMapper/Source/RigMapperEditor/Private/RigMapperGraph/RigMapperDefinitionEditorGraphNode.cpp
--- a/Engine/Plugins/Experimental/Animation/RigMapper/Source/RigMapperEditor/Private/RigMapperGraph/RigMapperDefinitionEditorGraphNode.cpp
+++ b/Engine/Plugins/Experimental/Animation/RigMapper/Source/RigMapperEditor/Private/RigMapperGraph/RigMapperDefinitionEditorGraphNode.cpp
@@ -257,6 +257,25 @@ void URigMapperDefinitionEditorGraphNode::RemoveLinksToFeature(const FString& In
 	FeaturePtr->RenameInput(InFeatureName, TEXT(""));
 }
 
+void URigMapperDefinitionEditorGraphNode::PinConnectionListChanged(UEdGraphPin* Pin)
+{
+	FRigMapperFeature* FeaturePtr = GetNodeFeature();
+	if (!FeaturePtr)
+	{
+		return;
+	}
+
+	if (!Pin)
+	{
+		return;
+	}
+
+	if (Pin->Direction == EGPD_Input && Pin->LinkedTo.Num() == 0)
+	{
+		RelinkFeature(TEXT(""), Pin);
+	}
+}
+
 FRigMapperFeature* URigMapperDefinitionEditorGraphNode::GetNodeFeature() const
 {
 	if (Definition.IsValid())
diff --git a/Engine/Plugins/Experimental/Animation/RigMapper/Source/RigMapperEditor/Private/RigMapperGraph/RigMapperDefinitionEditorGraphNode.h b/Engine/Plugins/Experimental/Animation/RigMapper/Source/RigMapperEditor/Private/RigMapperGraph/RigMapperDefinitionEditorGraphNode.h
--- a/Engine/Plugins/Experimental/Animation/RigMapper/Source/RigMapperEditor/Private/RigMapperGraph/RigMapperDefinitionEditorGraphNode.h
+++ b/Engine/Plugins/Experimental/Animation/RigMapper/Source/RigMapperEditor/Private/RigMapperGraph/RigMapperDefinitionEditorGraphNode.h
@@ -36,6 +36,7 @@ public:
 	UE_API virtual FLinearColor GetNodeTitleColor() const override;
 	UE_API virtual void GetNodeContextMenuActions(class UToolMenu* Menu, class UGraphNodeContextMenuContext* Context) const override;
 	UE_API virtual void AllocateDefaultPins() override;
+	void PinConnectionListChanged(UEdGraphPin* Pin) override;
 	// !UEdGraphNode implementation
 
 	UE_API void SetupNode(URigMapperDefinition* InDefinition, const FString& InNodeName, ERigMapperFeatureType InNodeType);
@@ -71,6 +72,8 @@ public:
 	UE_API void AddNewInputLink(const FString& InNewInput);
 	UE_API void RemoveLinksToFeature(const FString& InFeatureName);
 
+
+
 protected:
 	/** Cached title for the node */
 	FText NodeTitle;
diff --git a/Engine/Plugins/Experimental/Animation/RigMapper/Source/RigMapperEditor/Private/RigMapperGraph/RigMapperDefinitionEditorGraphSchema.cpp b/Engine/Plugins/Experimental/Animation/RigMapper/Source/RigMapperEditor/Private/RigMapperGraph/RigMapperDefinitionEditorGraphSchema.cpp
--- a/Engine/Plugins/Experimental/Animation/RigMapper/Source/RigMapperEditor/Private/RigMapperGraph/RigMapperDefinitionEditorGraphSchema.cpp
+++ b/Engine/Plugins/Experimental/Animation/RigMapper/Source/RigMapperEditor/Private/RigMapperGraph/RigMapperDefinitionEditorGraphSchema.cpp
@@ -114,14 +114,42 @@ bool URigMapperDefinitionEditorGraphSchema::SupportsDropPinOnNode(UEdGraphNode*
 			OutErrorMessage = LOCTEXT("RigMapperDefinitionEditor_CreateNewConnection_DropPinOnNode", "Create new input pin and connect.");
 			return true;
 		}
-		OutErrorMessage = LOCTEXT("RigMapperDefinitionEditor_CannotCreateConnection_DropPinOnNode", "It is not possible to create new connections on this node.");
+		OutErrorMessage = LOCTEXT("RigMapperDefinitionEditor_CannotCreateConnection_DropPinOnNode", "It is not possible to create new input connections on this node.");
 	}
 	return false;
 }
 
+void URigMapperDefinitionEditorGraphSchema::SetPinBeingDroppedOnNode(UEdGraphPin* InSourcePin) const
+{
+	PinBeingDroppedNodeName = TEXT("");
+	if (!InSourcePin)
+	{
+		return;
+	}
+	if (URigMapperDefinitionEditorGraphNode* SourceNode = Cast<URigMapperDefinitionEditorGraphNode>(InSourcePin->GetOwningNode()))
+	{
+		PinBeingDroppedNodeName = SourceNode->GetNodeName();
+	}
+}
+
 UEdGraphPin* URigMapperDefinitionEditorGraphSchema::DropPinOnNode(UEdGraphNode* InTargetNode, const FName& InSourcePinName, const FEdGraphPinType& InSourcePinType, EEdGraphPinDirection InSourcePinDirection) const
 {
-	return Super::DropPinOnNode(InTargetNode, InSourcePinName, InSourcePinType, InSourcePinDirection);
+	if (PinBeingDroppedNodeName.IsEmpty())
+	{
+		return nullptr;
+	}
+
+	UEdGraphPin* NewPin = nullptr;
+	if (InSourcePinDirection == EGPD_Output)
+	{
+		if (URigMapperDefinitionEditorGraphNode* GraphNode = Cast<URigMapperDefinitionEditorGraphNode>(InTargetNode))
+		{
+			NewPin = GraphNode->CreateInputPin();
+			GraphNode->AddNewInputLink(PinBeingDroppedNodeName);
+		}
+	}
+	PinBeingDroppedNodeName = TEXT("");
+	return NewPin;
 }
 
 void URigMapperDefinitionEditorGraphSchema::GetGraphContextActions(FGraphContextMenuBuilder& ContextMenuBuilder) const
@@ -295,10 +323,6 @@ bool URigMapperDefinitionEditorGraphSchema::CreatePromotedConnection(UEdGraphPin
 	return Super::CreatePromotedConnection(A, B);
 }
 
-void URigMapperDefinitionEditorGraphSchema::SetPinBeingDroppedOnNode(UEdGraphPin* InSourcePin) const
-{
-}
-
 bool URigMapperDefinitionEditorGraphSchema::TryRelinkConnectionTarget(UEdGraphPin* SourcePin, UEdGraphPin* OldTargetPin, UEdGraphPin* NewTargetPin, const TArray<UEdGraphNode*>& InSelectedGraphNodes) const
 {
 	return Super::TryRelinkConnectionTarget(SourcePin, OldTargetPin, NewTargetPin, InSelectedGraphNodes);
diff --git a/Engine/Plugins/Experimental/Animation/RigMapper/Source/RigMapperEditor/Private/RigMapperGraph/RigMapperDefinitionEditorGraphSchema.h b/Engine/Plugins/Experimental/Animation/RigMapper/Source/RigMapperEditor/Private/RigMapperGraph/RigMapperDefinitionEditorGraphSchema.h
--- a/Engine/Plugins/Experimental/Animation/RigMapper/Source/RigMapperEditor/Private/RigMapperGraph/RigMapperDefinitionEditorGraphSchema.h
+++ b/Engine/Plugins/Experimental/Animation/RigMapper/Source/RigMapperEditor/Private/RigMapperGraph/RigMapperDefinitionEditorGraphSchema.h
@@ -78,6 +78,10 @@ public:
 	UE_API bool TryRelinkConnectionTarget(UEdGraphPin* SourcePin, UEdGraphPin* OldTargetPin, UEdGraphPin* NewTargetPin, const TArray<UEdGraphNode*>& InSelectedGraphNodes) const override;
 	bool SafeDeleteNodeFromGraph(UEdGraph* Graph, UEdGraphNode* Node) const override;
 	// End of UEdGraphSchema interface
+
+private:
+	/** The name of the owning node of the pin that is being dropped on another node during a drag-drop operation. Controlled by SupportsDropPinOnNode(). */
+	mutable FString PinBeingDroppedNodeName = TEXT("");
 };
 
 #undef UE_API
