[RigMapper] Fixed a crash with Weighted Sum node.

diff --git a/Engine/Plugins/Experimental/Animation/RigMapper/Source/RigMapper/Private/RigMapperDefinition.cpp b/Engine/Plugins/Experimental/Animation/RigMapper/Source/RigMapper/Private/RigMapperDefinition.cpp
--- a/Engine/Plugins/Experimental/Animation/RigMapper/Source/RigMapper/Private/RigMapperDefinition.cpp
+++ b/Engine/Plugins/Experimental/Animation/RigMapper/Source/RigMapper/Private/RigMapperDefinition.cpp
@@ -1018,22 +1018,22 @@ void FRigMapperWsFeature::RelinkFeature(const FString& InNewInput, const int32 I
 	if (OrderedInputs.IsValidIndex(InputIndex))
 	{
 		const FString OldInput = OrderedInputs[InputIndex].Name;
+		double TempValue = 0.0f;
 		if (double* Weight = Inputs.Find(OldInput))
 		{
-			double TempValue = MoveTemp(*Weight);
+			TempValue = MoveTemp(*Weight);
 			Inputs.Remove(OldInput);
-			if (!Inputs.Contains(InNewInput))
-			{
-				Inputs.Add(InNewInput, TempValue);
-			}
-			else
-			{
-				Inputs[InNewInput] += TempValue;
-			}
-			OrderedInputs[InputIndex].Name = InNewInput;
 		}
+		if (!Inputs.Contains(InNewInput))
+		{
+			Inputs.Add(InNewInput, TempValue);
+		}
+		else
+		{
+			Inputs[InNewInput] += TempValue;
+		}
+		OrderedInputs[InputIndex].Name = InNewInput;
 	}
-
 }
 
 void FRigMapperWsFeature::AddNewInputLink(const FString& InNewInput)
@@ -1061,7 +1061,14 @@ void FRigMapperWsFeature::UpdateFromDetails()
 	Inputs.Empty();
 	for (FRigMapperWsInput& WsInput : OrderedInputs)
 	{
-		Inputs.Add(WsInput.Name, WsInput.Weight);
+		if (Inputs.Contains(WsInput.Name))
+		{
+			Inputs[WsInput.Name] += WsInput.Weight;
+		}
+		else
+		{
+			Inputs.Add(WsInput.Name, WsInput.Weight);
+		}
 	}
 }
 
diff --git a/Engine/Plugins/Experimental/Animation/RigMapper/Source/RigMapperEditor/Private/RigMapperGraph/RigMapperDefinitionEditorGraph.cpp b/Engine/Plugins/Experimental/Animation/RigMapper/Source/RigMapperEditor/Private/RigMapperGraph/RigMapperDefinitionEditorGraph.cpp
--- a/Engine/Plugins/Experimental/Animation/RigMapper/Source/RigMapperEditor/Private/RigMapperGraph/RigMapperDefinitionEditorGraph.cpp
+++ b/Engine/Plugins/Experimental/Animation/RigMapper/Source/RigMapperEditor/Private/RigMapperGraph/RigMapperDefinitionEditorGraph.cpp
@@ -254,6 +254,7 @@ URigMapperDefinitionEditorGraphNode* URigMapperDefinitionEditorGraph::ConstructF
 			Node->CreateInputPin();
 		}
 		
+		Node->UpdateNodeFields();
 		return Node;
 	}
 	return nullptr;
@@ -355,6 +356,7 @@ URigMapperDefinitionEditorGraphNode* URigMapperDefinitionEditorGraph::CreateGrap
 		FRigMapperWsFeature NewFeature = FRigMapperWsFeature(NodeName);
 		if (!InputNodeName.IsEmpty())
 		{
+			NewFeature.OrderedInputs.Add(FRigMapperWsInput(InputNodeName, 0.0));
 			NewFeature.Inputs.Add(InputNodeName);
 		}
 		WeakDefinition->Features.WeightedSums.Add(NewFeature);
@@ -411,6 +413,7 @@ URigMapperDefinitionEditorGraphNode* URigMapperDefinitionEditorGraph::CreateGrap
 			}
 		}
 	}
+	NewNode->UpdateNodeFields();
 	// Notify listeners that new Feature is created.
 	WeakDefinition->OnRigMapperDefinitionUpdated.Broadcast();
 
diff --git a/Engine/Plugins/Experimental/Animation/RigMapper/Source/RigMapperEditor/Private/RigMapperGraph/RigMapperDefinitionEditorGraphNode.cpp b/Engine/Plugins/Experimental/Animation/RigMapper/Source/RigMapperEditor/Private/RigMapperGraph/RigMapperDefinitionEditorGraphNode.cpp
--- a/Engine/Plugins/Experimental/Animation/RigMapper/Source/RigMapperEditor/Private/RigMapperGraph/RigMapperDefinitionEditorGraphNode.cpp
+++ b/Engine/Plugins/Experimental/Animation/RigMapper/Source/RigMapperEditor/Private/RigMapperGraph/RigMapperDefinitionEditorGraphNode.cpp
@@ -102,8 +102,6 @@ void URigMapperDefinitionEditorGraphNode::SetupNode(URigMapperDefinition* InDefi
 		FRigMapperNullOutputFeature NewFeature = FRigMapperNullOutputFeature(NodeName);
 		Definition->Features.NullOutputs.Add(NewFeature);
 	}
-
-	UpdateNodeFields();
 }
 
 void URigMapperDefinitionEditorGraphNode::UpdateNodeFields()
@@ -142,19 +140,18 @@ void URigMapperDefinitionEditorGraphNode::UpdateNodeFields()
 		if (FeaturePtr->GetFeatureType() == ERigMapperFeatureType::WeightedSum)
 		{
 			const FRigMapperWsFeature* Feature = static_cast<FRigMapperWsFeature*>(FeaturePtr);
-			const TMap<FString, double>& Inputs = Feature->Inputs;
+			const TArray<FRigMapperWsInput> Inputs = Feature->OrderedInputs;
 			const FRigMapperFeatureRange& Range = Feature->Range;
 			const FText MinText = LOCTEXT("RangeMinimum", "min");
 			const FText MaxText = LOCTEXT("RangeMaximum", "max");
 
 			for (UEdGraphPin* InPin : InputPins)
 			{
-				for (const UEdGraphPin* OutPin : InPin->LinkedTo)
+				for (int32 PinIndex = 0; PinIndex < InPin->LinkedTo.Num(); PinIndex++)
 				{
-					const FString LinkedNodeName = OutPin->GetOwningNode()->GetNodeTitle(ENodeTitleType::FullTitle).ToString();
-					if (Inputs.Contains(LinkedNodeName))
+					if (Inputs.IsValidIndex(PinIndex))
 					{
-						InPin->PinFriendlyName = FText::Format(INVTEXT("{0}"), Inputs[LinkedNodeName]);
+						InPin->PinFriendlyName = FText::Format(INVTEXT("{0}"), Inputs[PinIndex].Weight);
 						FString RangeStr;
 						if (Range.bHasLowerBound)
 						{
@@ -325,7 +322,8 @@ void URigMapperDefinitionEditorGraphNode::UpdateNodeFromFeature(const FPropertyC
 	}
 	FString OldName = NodeName;
 	FRigMapperFeature* FeaturePtr = nullptr;
-	if ((InPropertyChangedEvent.ChangeType & EPropertyChangeType::ValueSet) && InPropertyChangedEvent.GetPropertyName() == GET_MEMBER_NAME_CHECKED(FRigMapperFeature, Name))
+	if ((InPropertyChangedEvent.ChangeType & EPropertyChangeType::ValueSet) && InPropertyChangedEvent.GetPropertyName() == GET_MEMBER_NAME_CHECKED(FRigMapperFeature, Name) &&
+		InPropertyChangedEvent.GetMemberPropertyName() == GET_MEMBER_NAME_CHECKED(FRigMapperFeature, Name))
 	{
 		// Name has been modified, find feature using new name.
 		if (const FStrProperty* NameProp = CastField<FStrProperty>(InPropertyChangedEvent.Property))
@@ -352,16 +350,18 @@ void URigMapperDefinitionEditorGraphNode::UpdateNodeFromFeature(const FPropertyC
 	const FName PropertyName = InPropertyChangedEvent.GetPropertyName();
 	const FName ParentName = InPropertyChangedEvent.GetMemberPropertyName();
 	// TODO: Refactor with specialized node subclasses for each type
-	if (NodeType == ERigMapperFeatureType::Multiply || NodeType == ERigMapperFeatureType::WeightedSum)
+	if (NodeType == ERigMapperFeatureType::Multiply || NodeType == ERigMapperFeatureType::WeightedSum || NodeType == ERigMapperFeatureType::SDK || NodeType == ERigMapperFeatureType::Output)
 	{
 		TMap<FString, int32> PropertyNameStack;
 		InPropertyChangedEvent.GetArrayIndicesPerObject(0, PropertyNameStack);
+
 		const FString InputsProperty = NodeType == ERigMapperFeatureType::WeightedSum ?
 			GET_MEMBER_NAME_STRING_CHECKED(FRigMapperWsFeature, OrderedInputs) :
 			GET_MEMBER_NAME_STRING_CHECKED(FRigMapperMultiplyFeature, Inputs);
+
 		const int32 PinIndex = PropertyNameStack.Contains(InputsProperty) ? PropertyNameStack.FindRef(InputsProperty) : INDEX_NONE;
 
-		const int32 MinReqPins = NodeType == ERigMapperFeatureType::WeightedSum ? 1 : 2;
+		const int32 MinReqPins = NodeType == ERigMapperFeatureType::Multiply ? 2 : 1;
 
 		if (InPropertyChangedEvent.ChangeType & EPropertyChangeType::ArrayClear)
 		{	
@@ -397,62 +397,47 @@ void URigMapperDefinitionEditorGraphNode::UpdateNodeFromFeature(const FPropertyC
 		else if (InPropertyChangedEvent.ChangeType & EPropertyChangeType::ArrayAdd)
 		{
 			// ADDING A NEW (EMPTY) INPUT
-			if (NodeType == ERigMapperFeatureType::Multiply)
-			{
-				if (PinIndex != INDEX_NONE)
-				{
-					FRigMapperMultiplyFeature* Multiply = static_cast<FRigMapperMultiplyFeature*>(FeaturePtr);
-
-					if (Multiply->Inputs.Num() > InputPins.Num())
-					{
-						CreateInputPin();
-					}
-				}
-			}
-			else if (NodeType == ERigMapperFeatureType::WeightedSum)
+			if (PinIndex != INDEX_NONE) // This will be true only for WS and Multiply features
 			{
-				if (PinIndex != INDEX_NONE)
+				TArray<FString> Inputs;
+				FeaturePtr->GetInputs(Inputs);
+				if (Inputs.Num() > InputPins.Num())
 				{
-					FRigMapperWsFeature* WSFeature = static_cast<FRigMapperWsFeature*>(FeaturePtr);
-
-					if (WSFeature->OrderedInputs.Num() > InputPins.Num())
-					{
-						CreateInputPin();
-						WSFeature->UpdateFromDetails();
-					}
+					CreateInputPin();
+					FeaturePtr->UpdateFromDetails();
 				}
 			}
 		}
 		else if (InPropertyChangedEvent.ChangeType & EPropertyChangeType::ValueSet)
 		{
-			// CHANGING AN INPUT NAME
-			if (PinIndex != INDEX_NONE)
+			if ((PinIndex != INDEX_NONE && InPropertyChangedEvent.GetPropertyName() == InputsProperty) ||				// For WS and Multiply features
+				InPropertyChangedEvent.GetPropertyName() == GET_MEMBER_NAME_CHECKED(FRigMapperSdkFeature, Input))		// For SDK and Output features
 			{
+				// CHANGING AN INPUT NAME
 				URigMapperDefinitionEditorGraph* RigMapperGraph = Cast<URigMapperDefinitionEditorGraph>(GetGraph());
 				URigMapperDefinitionEditorGraphNode* LinkedNode = nullptr;
-				if (NodeType == ERigMapperFeatureType::Multiply)
-				{
-					if (FRigMapperMultiplyFeature* Multiply = static_cast<FRigMapperMultiplyFeature*>(FeaturePtr))
-					{
-						const FString InputName = Multiply->Inputs[PinIndex];
-						LinkedNode = RigMapperGraph->FindNodeByName(InputName);
-					}
-				}
-				else if (NodeType == ERigMapperFeatureType::WeightedSum)
+				TArray<FString> Inputs;
+				FeaturePtr->GetInputs(Inputs);
+				const int32 InputIndex = PinIndex == INDEX_NONE ? 0 : PinIndex;
+				if (InputPins.IsValidIndex(InputIndex))
 				{
-					if (FRigMapperWsFeature* WSFeature = static_cast<FRigMapperWsFeature*>(FeaturePtr))
+					InputPins[InputIndex]->BreakAllPinLinks();
+					if (Inputs.IsValidIndex(InputIndex))
 					{
-						const FString InputName = WSFeature->OrderedInputs[PinIndex].Name;
+						const FString InputName = Inputs[InputIndex];
 						LinkedNode = RigMapperGraph->FindNodeByName(InputName);
-						WSFeature->UpdateFromDetails();
+						if (LinkedNode)
+						{
+							InputPins[InputIndex]->MakeLinkTo(LinkedNode->GetOutputPin());
+						}
 					}
 				}
-				if (LinkedNode && InputPins.IsValidIndex(PinIndex))
-				{
-					InputPins[PinIndex]->BreakAllPinLinks();
-					InputPins[PinIndex]->MakeLinkTo(LinkedNode->GetOutputPin());
-				}
-
+				FeaturePtr->UpdateFromDetails();
+			}
+			else if (PinIndex != INDEX_NONE && NodeType == ERigMapperFeatureType::WeightedSum && InPropertyChangedEvent.GetPropertyName() == GET_MEMBER_NAME_CHECKED(FRigMapperWsInput, Weight))
+			{
+				// CHANGING WS INPUT WEIGHT
+				FeaturePtr->UpdateFromDetails();
 			}
 		}
 	}
diff --git a/Engine/Plugins/Experimental/Animation/RigMapper/Source/RigMapperEditor/Private/RigMapperGraph/SRigMapperDefinitionGraphEditor.cpp b/Engine/Plugins/Experimental/Animation/RigMapper/Source/RigMapperEditor/Private/RigMapperGraph/SRigMapperDefinitionGraphEditor.cpp
--- a/Engine/Plugins/Experimental/Animation/RigMapper/Source/RigMapperEditor/Private/RigMapperGraph/SRigMapperDefinitionGraphEditor.cpp
+++ b/Engine/Plugins/Experimental/Animation/RigMapper/Source/RigMapperEditor/Private/RigMapperGraph/SRigMapperDefinitionGraphEditor.cpp
@@ -124,6 +124,11 @@ void SRigMapperDefinitionGraphEditor::RefreshGraphNode(URigMapperDefinitionEdito
 	}
 }
 
+void SRigMapperDefinitionGraphEditor::CaptureKeyboard()
+{
+	GraphEditor->CaptureKeyboard();
+}
+
 void SRigMapperDefinitionGraphEditor::ZoomToFitNodes(const TArray<URigMapperDefinitionEditorGraphNode*>& SelectedNodes) const
 {
 	if (GraphObj->NeedsRefreshLayout())
diff --git a/Engine/Plugins/Experimental/Animation/RigMapper/Source/RigMapperEditor/Private/RigMapperGraph/SRigMapperDefinitionGraphEditor.h b/Engine/Plugins/Experimental/Animation/RigMapper/Source/RigMapperEditor/Private/RigMapperGraph/SRigMapperDefinitionGraphEditor.h
--- a/Engine/Plugins/Experimental/Animation/RigMapper/Source/RigMapperEditor/Private/RigMapperGraph/SRigMapperDefinitionGraphEditor.h
+++ b/Engine/Plugins/Experimental/Animation/RigMapper/Source/RigMapperEditor/Private/RigMapperGraph/SRigMapperDefinitionGraphEditor.h
@@ -34,6 +34,7 @@ public:
 	UE_API void SelectNodes(const TArray<FString>& Inputs, const TArray<FString>& Features, const TArray<FString>& Outputs, const TArray<FString>& NullOutputs);
 	UE_API void RebuildGraph();
 	UE_API void RefreshGraphNode(URigMapperDefinitionEditorGraphNode* InNode);
+	UE_API void CaptureKeyboard();
 
 private:
 	UE_API void ZoomToFitNodes(const TArray<URigMapperDefinitionEditorGraphNode*>& SelectedNodes) const;
