[RigMapper] Fixed a bug with weighted sum pin values.

diff --git a/Engine/Plugins/Experimental/Animation/RigMapper/Source/RigMapperEditor/Private/RigMapperGraph/RigMapperDefinitionEditorGraphNode.cpp b/Engine/Plugins/Experimental/Animation/RigMapper/Source/RigMapperEditor/Private/RigMapperGraph/RigMapperDefinitionEditorGraphNode.cpp
--- a/Engine/Plugins/Experimental/Animation/RigMapper/Source/RigMapperEditor/Private/RigMapperGraph/RigMapperDefinitionEditorGraphNode.cpp
+++ b/Engine/Plugins/Experimental/Animation/RigMapper/Source/RigMapperEditor/Private/RigMapperGraph/RigMapperDefinitionEditorGraphNode.cpp
@@ -145,31 +145,27 @@ void URigMapperDefinitionEditorGraphNode::UpdateNodeFields()
 			const FText MinText = LOCTEXT("RangeMinimum", "min");
 			const FText MaxText = LOCTEXT("RangeMaximum", "max");
 
-			for (UEdGraphPin* InPin : InputPins)
+			for (int32 PinIndex = 0; PinIndex < InputPins.Num(); PinIndex++)
 			{
-				for (int32 PinIndex = 0; PinIndex < InPin->LinkedTo.Num(); PinIndex++)
+				if (Inputs.IsValidIndex(PinIndex))
 				{
-					if (Inputs.IsValidIndex(PinIndex))
-					{
-						InPin->PinFriendlyName = FText::Format(INVTEXT("{0}"), Inputs[PinIndex].Weight);
-						FString RangeStr;
-						if (Range.bHasLowerBound)
-						{
-							RangeStr.Appendf(TEXT("%s: %.*f"), *MinText.ToString(), 3, Range.LowerBound);
-						}
-						if (Range.bHasLowerBound && Range.bHasUpperBound)
-						{
-							RangeStr.Append(TEXT("\n"));
-						}
-						if (Range.bHasUpperBound)
-						{
-							RangeStr.Appendf(TEXT("%s: %.*f"), *MaxText.ToString(), 3, Range.UpperBound);
-						}
-						NodeSubtitle = FText::FromString(RangeStr);
-					}
+					InputPins[PinIndex]->PinFriendlyName = FText::Format(INVTEXT("{0}"), Inputs[PinIndex].Weight);
 				}
 			}
-
+			FString RangeStr;
+			if (Range.bHasLowerBound)
+			{
+				RangeStr.Appendf(TEXT("%s: %.*f"), *MinText.ToString(), 3, Range.LowerBound);
+			}
+			if (Range.bHasLowerBound && Range.bHasUpperBound)
+			{
+				RangeStr.Append(TEXT("\n"));
+			}
+			if (Range.bHasUpperBound)
+			{
+				RangeStr.Appendf(TEXT("%s: %.*f"), *MaxText.ToString(), 3, Range.UpperBound);
+			}
+			NodeSubtitle = FText::FromString(RangeStr);
 		}
 	}
 }
