[RigMapper] Fixed the issue with focusing graph nodes

diff --git a/Engine/Plugins/Experimental/Animation/RigMapper/Source/RigMapperEditor/Private/RigMapperDefinitionEditorToolkit.cpp b/Engine/Plugins/Experimental/Animation/RigMapper/Source/RigMapperEditor/Private/RigMapperDefinitionEditorToolkit.cpp
--- a/Engine/Plugins/Experimental/Animation/RigMapper/Source/RigMapperEditor/Private/RigMapperDefinitionEditorToolkit.cpp
+++ b/Engine/Plugins/Experimental/Animation/RigMapper/Source/RigMapperEditor/Private/RigMapperDefinitionEditorToolkit.cpp
@@ -13,6 +13,7 @@
 #include "PropertyEditorModule.h"
 #include "Modules/ModuleManager.h"
 #include "RigMapperGraph/SRigMapperDefinitionStructureView.h"
+#include "Framework/Application/SlateApplication.h"
 #include "Widgets/Docking/SDockTab.h"
 #include "Widgets/Views/STreeView.h"
 
@@ -156,14 +157,14 @@ TSharedRef<SDockTab> FRigMapperDefinitionEditorToolkit::SpawnGraphTab(const FSpa
 		Definition->OnRigMapperDefinitionCreated.AddRaw(this, &FRigMapperDefinitionEditorToolkit::HandleRigMapperDefinitionLoaded);
 		Definition->OnRigMapperDefinitionUpdated.AddRaw(this, &FRigMapperDefinitionEditorToolkit::HandleRigMapperDefinitionUpdated);
 	}
-	TSharedRef<SDockTab> DockTab = SNew(SDockTab)
+	GraphTab = SNew(SDockTab)
 	.TabColorScale(GetTabColorScale())
 	.Label(LOCTEXT("RigMapperDefinitionEditorToolkitGraphTab", "Graph"))
 	[
 		GraphEditor.ToSharedRef()
 	];
-	
-	return DockTab;
+
+	return GraphTab.ToSharedRef();
 }
 
 void FRigMapperDefinitionEditorToolkit::HandleRigMapperDefinitionLoaded()
@@ -291,86 +292,6 @@ bool FRigMapperDefinitionEditorToolkit::HandleIsPropertyVisible(const FPropertyA
 	return true;
 }
 
-void FRigMapperDefinitionEditorToolkit::HandleFinishedChangingProperties(const FPropertyChangedEvent& PropertyChangedEvent)
-{
-	FPropertyChangedEvent Event = PropertyChangedEvent;
-	Event.ObjectIteratorIndex = 0;
-	
-	const EPropertyChangeType::Type ChangeType = Event.ChangeType;
-	const FName PropertyName = Event.GetPropertyName();
-	const FName ParentName = Event.GetMemberPropertyName();
-	const FName InputName = GET_MEMBER_NAME_CHECKED(URigMapperDefinition, Inputs);
-	const FName OutputName = GET_MEMBER_NAME_CHECKED(URigMapperDefinition, Outputs);
-	const FName FeaturesName = GET_MEMBER_NAME_CHECKED(URigMapperDefinition, Features);
-	
-	const FName NullOutputName = GET_MEMBER_NAME_CHECKED(URigMapperDefinition, NullOutputs);
-
-	const bool bInput = PropertyName == InputName && ParentName == InputName;
-	const bool bOutput = !bInput && (PropertyName == OutputName || PropertyName == OutputName.ToString() + TEXT("_Key")) && ParentName == OutputName; // PropertyName will be Outputs_Key if renaming an output
-	const bool bName = !bInput && !bOutput && PropertyName == GET_MEMBER_NAME_CHECKED(FRigMapperFeature, Name);
-	const bool bFeaturesChildProperty = ParentName == FeaturesName;
-	const bool bSelectNew = (ChangeType == EPropertyChangeType::Duplicate || ChangeType == EPropertyChangeType::ArrayAdd);
-	const bool bNullOutput = PropertyName == NullOutputName && ParentName == NullOutputName;
-	const bool bRename = ChangeType == EPropertyChangeType::ValueSet && (bInput || bOutput || bName || bFeaturesChildProperty || bNullOutput);
-	
-	//GraphEditor->RebuildGraph();
-	StructureView->RebuildTree();	
-	if (bSelectNew || bRename)
-	{
-		if (bInput)
-		{
-			int32 InputIndex = Event.GetArrayIndex(PropertyName.ToString());
-			if (Definition->Inputs.IsValidIndex(InputIndex))
-			{
-				StructureView->SelectNode(Definition->Inputs[InputIndex], ERigMapperFeatureType::Input, true);
-				GraphEditor->SelectNodes({ Definition->Inputs[InputIndex] }, {}, {}, {});
-			}
-		}
-		else if (bOutput)
-		{
-			int32 OutputIndex = Event.GetArrayIndex(PropertyName.ToString());
-			TArray<FString> OutputsInOrder;
-			Definition->Outputs.GenerateKeyArray(OutputsInOrder);
-			if (OutputsInOrder.IsValidIndex(OutputIndex))
-			{
-				StructureView->SelectNode(OutputsInOrder[OutputIndex], ERigMapperFeatureType::Output, true);
-				GraphEditor->SelectNodes({}, {}, { OutputsInOrder[OutputIndex] }, {});
-			}
-		}
-		if (bNullOutput)
-		{
-			int32 NullOutputIndex = Event.GetArrayIndex(PropertyName.ToString());
-			if (Definition->NullOutputs.IsValidIndex(NullOutputIndex))
-			{
-				StructureView->SelectNode(Definition->NullOutputs[NullOutputIndex], ERigMapperFeatureType::NullOutput, true);
-				GraphEditor->SelectNodes({ Definition->NullOutputs[NullOutputIndex] }, {}, {}, {});
-			}
-		}
-		else if (bFeaturesChildProperty)
-		{
-			if (const int32 MultiplyIndex = Event.GetArrayIndex(GET_MEMBER_NAME_CHECKED(FRigMapperFeatureDefinitions, Multiply).ToString()); MultiplyIndex != INDEX_NONE)
-			{
-				StructureView->SelectNode(Definition->Features.Multiply[MultiplyIndex].Name, ERigMapperFeatureType::Multiply, true);
-				GraphEditor->SelectNodes({}, { Definition->Features.Multiply[MultiplyIndex].Name }, {}, {});
-			}
-			else if (const int32 WsIndex = Event.GetArrayIndex(GET_MEMBER_NAME_CHECKED(FRigMapperFeatureDefinitions, WeightedSums).ToString()); WsIndex != INDEX_NONE)
-			{
-				StructureView->SelectNode(Definition->Features.WeightedSums[WsIndex].Name, ERigMapperFeatureType::WeightedSum, true);
-				GraphEditor->SelectNodes({}, { Definition->Features.WeightedSums[WsIndex].Name }, { }, {});
-			}
-			else if (const int32 SdkIndex = Event.GetArrayIndex(GET_MEMBER_NAME_CHECKED(FRigMapperFeatureDefinitions, SDKs).ToString()); SdkIndex != INDEX_NONE)
-			{
-				StructureView->SelectNode(Definition->Features.SDKs[SdkIndex].Name, ERigMapperFeatureType::SDK, true);
-				GraphEditor->SelectNodes({}, { Definition->Features.SDKs[SdkIndex].Name }, {}, {});
-			}
-		}
-	}
-
-	// todo: restore selection
-	
-	// todo: remove from selection if even remove / clear
-}
-
 void FRigMapperDefinitionEditorToolkit::HandleGraphSelectionChanged(const TSet<UObject*>& Nodes)
 {
 	if (!StructureView || !FeatureDetailsView)
@@ -391,6 +312,16 @@ void FRigMapperDefinitionEditorToolkit::HandleGraphSelectionChanged(const TSet<U
 			}					
 		}
 	}
+	if (GraphTab.IsValid())
+	{
+		FGlobalTabmanager::Get()->SetActiveTab(GraphTab);
+	}
+
+	FSlateApplication::Get().ClearKeyboardFocus(EFocusCause::SetDirectly);
+	if (GraphEditor)
+	{
+		GraphEditor->CaptureKeyboard();		
+	}
 }
 
 void FRigMapperDefinitionEditorToolkit::HandleStructureSelectionChanged(ESelectInfo::Type SelectInfo, TArray<FString> SelectedInputs, TArray<FString> SelectedFeatures, TArray<FString> SelectedOutputs, TArray<FString> SelectedNullOutputs)
diff --git a/Engine/Plugins/Experimental/Animation/RigMapper/Source/RigMapperEditor/Private/RigMapperDefinitionEditorToolkit.h b/Engine/Plugins/Experimental/Animation/RigMapper/Source/RigMapperEditor/Private/RigMapperDefinitionEditorToolkit.h
--- a/Engine/Plugins/Experimental/Animation/RigMapper/Source/RigMapperEditor/Private/RigMapperDefinitionEditorToolkit.h
+++ b/Engine/Plugins/Experimental/Animation/RigMapper/Source/RigMapperEditor/Private/RigMapperDefinitionEditorToolkit.h
@@ -57,7 +57,6 @@ private:
 	UE_API TSharedRef<SDockTab> SpawnStructureTab(const FSpawnTabArgs& SpawnTabArgs);
 
 	UE_API bool HandleIsPropertyVisible(const FPropertyAndParent& PropertyAndParent);
-	UE_API void HandleFinishedChangingProperties(const FPropertyChangedEvent& PropertyChangedEvent);
 	UE_API void HandleFinishedChangingFeatureProperties(const FPropertyChangedEvent& PropertyChangedEvent, TWeakObjectPtr<URigMapperDefinitionEditorGraphNode> InNode, TWeakPtr<FStructOnScope> InStructContainer);
 	UE_API void HandleRigMapperDefinitionLoaded();
 	UE_API void HandleRigMapperDefinitionUpdated();
@@ -79,6 +78,7 @@ private:
 	TSharedPtr<IDetailsView> DetailsView;
 	TSharedPtr<IStructureDetailsView> FeatureDetailsView;
 	TSharedPtr<SRigMapperDefinitionStructureView> StructureView;
+	TSharedPtr<SDockTab> GraphTab;
 
 };
 
diff --git a/Engine/Plugins/Experimental/Animation/RigMapper/Source/RigMapperEditor/Private/RigMapperGraph/SRigMapperDefinitionGraphEditor.cpp b/Engine/Plugins/Experimental/Animation/RigMapper/Source/RigMapperEditor/Private/RigMapperGraph/SRigMapperDefinitionGraphEditor.cpp
--- a/Engine/Plugins/Experimental/Animation/RigMapper/Source/RigMapperEditor/Private/RigMapperGraph/SRigMapperDefinitionGraphEditor.cpp
+++ b/Engine/Plugins/Experimental/Animation/RigMapper/Source/RigMapperEditor/Private/RigMapperGraph/SRigMapperDefinitionGraphEditor.cpp
@@ -8,12 +8,18 @@
 #include "RigMapperDefinition.h"
 #include "GraphEditorActions.h"
 #include "ScopedTransaction.h"
-#include "Framework/Commands/GenericCommands.h"
 
 #include "SGraphPanel.h"
 #include "SlateOptMacros.h"
+#include "Framework/Commands/GenericCommands.h"
+#include "Framework/Application/SlateApplication.h"
 #include "Components/VerticalBox.h"
+#include "Layout/WidgetPath.h"
 
+#if WITH_EDITOR
+#include "Editor/EditorEngine.h"
+extern UNREALED_API UEditorEngine* GEditor;
+#endif
 BEGIN_SLATE_FUNCTION_BUILD_OPTIMIZATION
 
 #define LOCTEXT_NAMESPACE "SRigMapperDefinitionGraphEditor"
@@ -104,7 +110,7 @@ void SRigMapperDefinitionGraphEditor::SelectNodes(const TArray<FString>& Inputs,
 		{
 			GraphEditor->SetNodeSelection(Node, true);
 		}
-		//ZoomToFitNodes(SelectedNodes);
+		ZoomToFitNodes(SelectedNodes);
 		
 		bSelectingNodes = false;
 	}
@@ -126,7 +132,28 @@ void SRigMapperDefinitionGraphEditor::RefreshGraphNode(URigMapperDefinitionEdito
 
 void SRigMapperDefinitionGraphEditor::CaptureKeyboard()
 {
-	GraphEditor->CaptureKeyboard();
+#if WITH_EDITOR
+	FTimerHandle TimerHandle;
+	GEditor->GetTimerManager()->SetTimer(TimerHandle, [WeakGraphEditor = GraphEditor.ToWeakPtr()]()
+		{
+			if (WeakGraphEditor.IsValid())
+			{
+				SGraphPanel* GraphPanel = WeakGraphEditor.Pin()->GetGraphPanel();
+				if (GraphPanel)
+				{
+					FWidgetPath SlatePath;
+					if (FSlateApplication::Get().GeneratePathToWidgetUnchecked(GraphPanel->AsShared(), SlatePath))
+					{
+						FSlateApplication::Get().SetKeyboardFocus(SlatePath, EFocusCause::SetDirectly);
+					}
+					else
+					{
+						WeakGraphEditor.Pin()->CaptureKeyboard();
+					}
+				}
+			}
+		}, 0.01f, false);
+#endif
 }
 
 void SRigMapperDefinitionGraphEditor::ZoomToFitNodes(const TArray<URigMapperDefinitionEditorGraphNode*>& SelectedNodes) const
diff --git a/Engine/Plugins/Experimental/Animation/RigMapper/Source/RigMapperEditor/Private/RigMapperGraph/SRigMapperDefinitionStructureView.cpp b/Engine/Plugins/Experimental/Animation/RigMapper/Source/RigMapperEditor/Private/RigMapperGraph/SRigMapperDefinitionStructureView.cpp
--- a/Engine/Plugins/Experimental/Animation/RigMapper/Source/RigMapperEditor/Private/RigMapperGraph/SRigMapperDefinitionStructureView.cpp
+++ b/Engine/Plugins/Experimental/Animation/RigMapper/Source/RigMapperEditor/Private/RigMapperGraph/SRigMapperDefinitionStructureView.cpp
@@ -31,6 +31,7 @@ void SRigMapperDefinitionStructureView::Construct(const FArguments& InArgs, URig
 	TreeView = SNew(STreeView<TSharedPtr<FString>>)
 		.SelectionMode(ESelectionMode::Multi)
 		.HighlightParentNodesForSelection(true)
+		.ReturnFocusToSelection(false)
 		.OnGenerateRow(this, &SRigMapperDefinitionStructureView::OnGenerateTreeRow)
 		.OnGetChildren(this, &SRigMapperDefinitionStructureView::OnGetTreeNodeChildren)
 		.OnSelectionChanged(this, &SRigMapperDefinitionStructureView::HandleTreeNodesSelectionChanged)
