[RigMapper] Fixed the issue with stale Input node names

#Synced-CL 48208203

diff --git a/Engine/Plugins/Experimental/Animation/RigMapper/Source/RigMapperEditor/Private/RigMapperGraph/RigMapperDefinitionEditorGraph.cpp b/Engine/Plugins/Experimental/Animation/RigMapper/Source/RigMapperEditor/Private/RigMapperGraph/RigMapperDefinitionEditorGraph.cpp
--- a/Engine/Plugins/Experimental/Animation/RigMapper/Source/RigMapperEditor/Private/RigMapperGraph/RigMapperDefinitionEditorGraph.cpp
+++ b/Engine/Plugins/Experimental/Animation/RigMapper/Source/RigMapperEditor/Private/RigMapperGraph/RigMapperDefinitionEditorGraph.cpp
@@ -100,10 +100,6 @@ URigMapperDefinitionEditorGraphNode* URigMapperDefinitionEditorGraph::FindNodeBy
 	{
 		return FoundNode;
 	}
-	if (URigMapperDefinitionEditorGraphNode* FoundNode = OutputNodes.FindRef(InNodeName))
-	{
-		return FoundNode;
-	}
 	if (URigMapperDefinitionEditorGraphNode* FoundNode = NullOutputNodes.FindRef(InNodeName))
 	{
 		return FoundNode;
diff --git a/Engine/Plugins/Experimental/Animation/RigMapper/Source/RigMapperEditor/Private/RigMapperGraph/RigMapperDefinitionEditorGraphNode.cpp b/Engine/Plugins/Experimental/Animation/RigMapper/Source/RigMapperEditor/Private/RigMapperGraph/RigMapperDefinitionEditorGraphNode.cpp
--- a/Engine/Plugins/Experimental/Animation/RigMapper/Source/RigMapperEditor/Private/RigMapperGraph/RigMapperDefinitionEditorGraphNode.cpp
+++ b/Engine/Plugins/Experimental/Animation/RigMapper/Source/RigMapperEditor/Private/RigMapperGraph/RigMapperDefinitionEditorGraphNode.cpp
@@ -208,7 +208,7 @@ void URigMapperDefinitionEditorGraphNode::RelinkFeature(const FString& InNewInpu
 		{
 			Definition->Outputs[NodeName] = InNewInput;
 		}
-	}	
+	}
 }
 
 void URigMapperDefinitionEditorGraphNode::AddNewInputLink(const FString& InNewInput)
@@ -438,19 +438,50 @@ void URigMapperDefinitionEditorGraphNode::UpdateNodeFromFeature(const FPropertyC
 		}
 	}
 
-	if (NodeName != OldName && NodeType != ERigMapperFeatureType::Output && NodeType != ERigMapperFeatureType::NullOutput) // Output and NullOutput nodes don't have output pins.
+	if (NodeName != OldName)
 	{
-		// Update inputs of all nodes directly connected to this node.
-		for (UEdGraphPin* OutPin : OutputPins) // It will always be one output pin.
+		if (NodeType != ERigMapperFeatureType::Output && NodeType != ERigMapperFeatureType::NullOutput) // Output and NullOutput nodes don't have output pins.
 		{
-			for (UEdGraphPin* LinkedPin : OutPin->LinkedTo)
+			// Update inputs of all nodes directly connected to this node.
+			for (UEdGraphPin* OutPin : OutputPins) // It will always be one output pin.
 			{
-				if (URigMapperDefinitionEditorGraphNode* LinkedNode = Cast<URigMapperDefinitionEditorGraphNode>(LinkedPin->GetOwningNode()))
+				for (UEdGraphPin* LinkedPin : OutPin->LinkedTo)
 				{
-					LinkedNode->RelinkFeature(NodeName, LinkedPin);
+					if (URigMapperDefinitionEditorGraphNode* LinkedNode = Cast<URigMapperDefinitionEditorGraphNode>(LinkedPin->GetOwningNode()))
+					{
+						LinkedNode->RelinkFeature(NodeName, LinkedPin);
+					}
 				}
 			}
 		}
+		// Input, Output and NullOutput need to update their Definition data from Feature.
+		if (NodeType == ERigMapperFeatureType::Input)
+		{
+			int32 OutIndex;
+			Definition->Inputs.Find(OldName, OutIndex);
+			if (OutIndex != INDEX_NONE)
+			{
+				Definition->Inputs[OutIndex] = NodeName;
+			}
+		}
+		else if (NodeType == ERigMapperFeatureType::Output)
+		{
+			if (FString* Input = Definition->Outputs.Find(OldName))
+			{
+				FString TempInput = MoveTemp(*Input);
+				Definition->Outputs.Remove(OldName);
+				Definition->Outputs.Add(NodeName, TempInput);
+			}			
+		}
+		else if (NodeType == ERigMapperFeatureType::NullOutput)
+		{
+			int32 OutIndex;
+			Definition->NullOutputs.Find(OldName, OutIndex);
+			if (OutIndex != INDEX_NONE)
+			{
+				Definition->NullOutputs[OutIndex] = NodeName;
+			}
+		}
 		// Update node map.
 		URigMapperDefinitionEditorGraph* RigMapperGraph = Cast<URigMapperDefinitionEditorGraph>(GetGraph());
 		RigMapperGraph->UpdateNodeNameInMap(OldName, NodeName, NodeType);
