[RigMapper] Fixed minor bugs with unique name generation and link management in RigMapperDefinition graph nodes

#Synced-CL 48280024

diff --git a/Engine/Plugins/Experimental/Animation/RigMapper/Source/RigMapper/Private/RigMapperUtils.cpp b/Engine/Plugins/Experimental/Animation/RigMapper/Source/RigMapper/Private/RigMapperUtils.cpp
--- a/Engine/Plugins/Experimental/Animation/RigMapper/Source/RigMapper/Private/RigMapperUtils.cpp
+++ b/Engine/Plugins/Experimental/Animation/RigMapper/Source/RigMapper/Private/RigMapperUtils.cpp
@@ -8,12 +8,14 @@ const TMap<ERigMapperFeatureType, FString> FeatureTypePrefixMap = {
 	{ ERigMapperFeatureType::WeightedSum, TEXT("ws:") },
 	{ ERigMapperFeatureType::Multiply, TEXT("mul:") },
 	{ ERigMapperFeatureType::SDK, TEXT("sdk:") },
-	{ ERigMapperFeatureType::Input, TEXT("") },
-	{ ERigMapperFeatureType::Output, TEXT("") },
-	{ ERigMapperFeatureType::NullOutput, TEXT("") },
+	{ ERigMapperFeatureType::Input, TEXT("in:") },
+	{ ERigMapperFeatureType::Output, TEXT("out:") },
+	{ ERigMapperFeatureType::NullOutput, TEXT("null:") },
 	{ ERigMapperFeatureType::Invalid, TEXT("invalid:") }
 };
 
+const FString GenericNodeName = TEXT("node");
+
 FString FRigMapperUtils::GenerateUniqueFeatureName(const TWeakObjectPtr<URigMapperDefinition> InDefinition, FString InDesiredName, ERigMapperFeatureType InFeatureType)
 {
 	TSet<FString> ExistingNames;
@@ -21,19 +23,6 @@ FString FRigMapperUtils::GenerateUniqueFeatureName(const TWeakObjectPtr<URigMapp
 	{
 		InDefinition.Pin()->Outputs.GetKeys(ExistingNames);
 		ExistingNames.Append(InDefinition.Pin()->NullOutputs);
-		bool bTrimmed = false;
-		if (InDesiredName.RemoveFromStart(FeatureTypePrefixMap[ERigMapperFeatureType::Multiply]))
-		{
-			bTrimmed = true;
-		}
-		if (!bTrimmed && InDesiredName.RemoveFromStart(FeatureTypePrefixMap[ERigMapperFeatureType::WeightedSum]))
-		{
-			bTrimmed = true;
-		}
-		if (!bTrimmed)
-		{
-			InDesiredName.RemoveFromStart(FeatureTypePrefixMap[ERigMapperFeatureType::SDK]);
-		}
 	}
 	else
 	{
@@ -50,27 +39,30 @@ FString FRigMapperUtils::GenerateUniqueFeatureName(const TWeakObjectPtr<URigMapp
 		{
 			ExistingNames.Add(Feature.Name);
 		}
-
-		InDesiredName = [&InDesiredName, InFeatureType]()
-			{
-				const FString& Prefix = FeatureTypePrefixMap[InFeatureType];
-				if (InDesiredName.IsEmpty() || InDesiredName.StartsWith(Prefix))
-				{
-					return InDesiredName;
-				}
-				if ((InFeatureType != ERigMapperFeatureType::Multiply && InDesiredName.RemoveFromStart(FeatureTypePrefixMap[ERigMapperFeatureType::Multiply])) ||
-					(InFeatureType != ERigMapperFeatureType::WeightedSum && InDesiredName.RemoveFromStart(FeatureTypePrefixMap[ERigMapperFeatureType::WeightedSum])) ||
-					(InFeatureType != ERigMapperFeatureType::SDK && InDesiredName.RemoveFromStart(FeatureTypePrefixMap[ERigMapperFeatureType::SDK])))
+		
+		if (InFeatureType != ERigMapperFeatureType::Input)
+		{
+			InDesiredName = [&InDesiredName, InFeatureType]()
 				{
+					const FString& Prefix = FeatureTypePrefixMap[InFeatureType];
+					if (InDesiredName.IsEmpty() || InDesiredName.StartsWith(Prefix))
+					{
+						return InDesiredName;
+					}
+					if ((InFeatureType != ERigMapperFeatureType::Multiply && InDesiredName.RemoveFromStart(FeatureTypePrefixMap[ERigMapperFeatureType::Multiply])) ||
+						(InFeatureType != ERigMapperFeatureType::WeightedSum && InDesiredName.RemoveFromStart(FeatureTypePrefixMap[ERigMapperFeatureType::WeightedSum])) ||
+						(InFeatureType != ERigMapperFeatureType::SDK && InDesiredName.RemoveFromStart(FeatureTypePrefixMap[ERigMapperFeatureType::SDK])))
+					{
+						return Prefix + InDesiredName;
+					}
 					return Prefix + InDesiredName;
-				}
-				return Prefix + InDesiredName;
-			}();
+				}();
+		}
 	}
 
 	if (InDesiredName.IsEmpty())
 	{
-		InDesiredName = FeatureTypePrefixMap[InFeatureType] + TEXT("node");
+		InDesiredName = FeatureTypePrefixMap[InFeatureType] + GenericNodeName;
 	}
 
 	// If the desired name is unique, just use that.
diff --git a/Engine/Plugins/Experimental/Animation/RigMapper/Source/RigMapperEditor/Private/RigMapperDefinitionEditorToolkit.cpp b/Engine/Plugins/Experimental/Animation/RigMapper/Source/RigMapperEditor/Private/RigMapperDefinitionEditorToolkit.cpp
--- a/Engine/Plugins/Experimental/Animation/RigMapper/Source/RigMapperEditor/Private/RigMapperDefinitionEditorToolkit.cpp
+++ b/Engine/Plugins/Experimental/Animation/RigMapper/Source/RigMapperEditor/Private/RigMapperDefinitionEditorToolkit.cpp
@@ -334,6 +334,11 @@ void FRigMapperDefinitionEditorToolkit::HandleStructureSelectionChanged(ESelectI
 
 void FRigMapperDefinitionEditorToolkit::SetFeatureDetailsView(const URigMapperDefinitionEditorGraphNode* InNode)
 {
+	if (!FeatureDetailsView)
+	{
+		return;
+	}
+
 	TObjectPtr<URigMapperDefinitionEditorGraphNode> Node = ConstCast<const URigMapperDefinitionEditorGraphNode>(InNode);
 	if (!Node)
 	{
@@ -388,6 +393,7 @@ void FRigMapperDefinitionEditorToolkit::HandleFinishedChangingFeatureProperties(
 		{
 			StructureView->RebuildTree();
 		}
+		SetFeatureDetailsView(InNode.Get());
 	}	
 
 }
diff --git a/Engine/Plugins/Experimental/Animation/RigMapper/Source/RigMapperEditor/Private/RigMapperGraph/RigMapperDefinitionEditorGraph.cpp b/Engine/Plugins/Experimental/Animation/RigMapper/Source/RigMapperEditor/Private/RigMapperGraph/RigMapperDefinitionEditorGraph.cpp
--- a/Engine/Plugins/Experimental/Animation/RigMapper/Source/RigMapperEditor/Private/RigMapperGraph/RigMapperDefinitionEditorGraph.cpp
+++ b/Engine/Plugins/Experimental/Animation/RigMapper/Source/RigMapperEditor/Private/RigMapperGraph/RigMapperDefinitionEditorGraph.cpp
@@ -90,7 +90,7 @@ TArray<URigMapperDefinitionEditorGraphNode*> URigMapperDefinitionEditorGraph::Ge
 }
 
 // TODO: Add RigMapperDefinition node subclasses
-URigMapperDefinitionEditorGraphNode* URigMapperDefinitionEditorGraph::FindNodeByName(const FString& InNodeName)
+URigMapperDefinitionEditorGraphNode* URigMapperDefinitionEditorGraph::FindFeatureNodeByName(const FString& InNodeName)
 {
 	if (URigMapperDefinitionEditorGraphNode* FoundNode = FeatureNodes.FindRef(InNodeName))
 	{
@@ -100,6 +100,15 @@ URigMapperDefinitionEditorGraphNode* URigMapperDefinitionEditorGraph::FindNodeBy
 	{
 		return FoundNode;
 	}
+	return nullptr;
+}
+
+URigMapperDefinitionEditorGraphNode* URigMapperDefinitionEditorGraph::FindOutputNodeByName(const FString& InNodeName)
+{
+	if (URigMapperDefinitionEditorGraphNode* FoundNode = OutputNodes.FindRef(InNodeName))
+	{
+		return FoundNode;
+	}
 	if (URigMapperDefinitionEditorGraphNode* FoundNode = NullOutputNodes.FindRef(InNodeName))
 	{
 		return FoundNode;
@@ -154,7 +163,10 @@ bool URigMapperDefinitionEditorGraph::UpdateNodeNameInMap(const FString& InOldNa
 
 bool URigMapperDefinitionEditorGraph::RemoveNodeFromMap(const FString& InNodeName, const ERigMapperFeatureType InNodeType)
 {
-	WeakDefinition->DeleteFeature(InNodeName, InNodeType);
+	if (WeakDefinition.IsValid())
+	{
+		WeakDefinition->DeleteFeature(InNodeName, InNodeType);
+	}
 	switch (InNodeType)
 	{
 	case ERigMapperFeatureType::WeightedSum:
@@ -215,6 +227,12 @@ URigMapperDefinitionEditorGraphNode* URigMapperDefinitionEditorGraph::ConstructO
 		{
 			LinkGraphNodes(LinkedNode, Node);
 		}
+		else
+		{
+			// Output node is saved without connection.
+			UEdGraphPin* InPin = Node->CreateInputPin();
+			InPin->bHidden = false;
+		}
 
 		return Node;
 	}
@@ -239,6 +257,12 @@ URigMapperDefinitionEditorGraphNode* URigMapperDefinitionEditorGraph::ConstructF
 			{
 				LinkGraphNodes(LinkedNode, Node);
 			}
+			else
+			{
+				// Feature node is saved without connected input pin.
+				UEdGraphPin* InPin = Node->CreateInputPin();
+				InPin->bHidden = false;
+			}
 		}
 
 		if (FeatureInputNames.Num() == 0)
@@ -324,7 +348,13 @@ URigMapperDefinitionEditorGraphNode* URigMapperDefinitionEditorGraph::CreateGrap
 	NewNode->NodePosX = InLocation.X;
 	NewNode->NodePosY = InLocation.Y;
 	// Create unique feature (node) name.
-	FString NodeName = FRigMapperUtils::GenerateUniqueFeatureName(WeakDefinition, FromPinName, InNodeType);
+	FString NodeName = TEXT("");
+	if (InNodeType != ERigMapperFeatureType::Input && InNodeType != ERigMapperFeatureType::Output && InNodeType != ERigMapperFeatureType::NullOutput)
+	{
+		// Only feature types can propagate names in chains.
+		NodeName = FromPinName;
+	}
+	NodeName = FRigMapperUtils::GenerateUniqueFeatureName(WeakDefinition, NodeName, InNodeType);
 	if (InNodeType != ERigMapperFeatureType::Input && InNodeType != ERigMapperFeatureType::NullOutput)
 	{
 		UEdGraphPin* InPin = NewNode->CreateInputPin();
@@ -404,6 +434,7 @@ URigMapperDefinitionEditorGraphNode* URigMapperDefinitionEditorGraph::CreateGrap
 			UEdGraphPin* OutPin = NewNode->GetOutputPin();
 			if (OutPin)
 			{
+				InFromPin->BreakAllPinLinks();
 				OutPin->MakeLinkTo(InFromPin);
 				Cast<URigMapperDefinitionEditorGraphNode>(InFromPin->GetOwningNode())->RelinkFeature(NodeName, InFromPin);
 			}
diff --git a/Engine/Plugins/Experimental/Animation/RigMapper/Source/RigMapperEditor/Private/RigMapperGraph/RigMapperDefinitionEditorGraph.h b/Engine/Plugins/Experimental/Animation/RigMapper/Source/RigMapperEditor/Private/RigMapperGraph/RigMapperDefinitionEditorGraph.h
--- a/Engine/Plugins/Experimental/Animation/RigMapper/Source/RigMapperEditor/Private/RigMapperGraph/RigMapperDefinitionEditorGraph.h
+++ b/Engine/Plugins/Experimental/Animation/RigMapper/Source/RigMapperEditor/Private/RigMapperGraph/RigMapperDefinitionEditorGraph.h
@@ -34,8 +34,10 @@ public:
 
 	UE_API URigMapperDefinitionEditorGraphNode* CreateGraphNode(const ERigMapperFeatureType InNodeType, UEdGraphPin* InFromPin, const FVector2f& InLocation, bool bSelectNewNode = true);
 	UE_API TArray<URigMapperDefinitionEditorGraphNode*> GetNodesByName(const TArray<FString>& Inputs, const TArray<FString>& Features, const TArray<FString>& Outputs, const TArray<FString>& NullOutputs) const;
-	/** Retrieves first node that matches the InNodeName (assumes node names are unique	*/
-	UE_API URigMapperDefinitionEditorGraphNode* FindNodeByName(const FString& InNodeName);
+	/** Retrieves first feature node that matches the InNodeName (assumes node names are unique and includes input nodes) */
+	UE_API URigMapperDefinitionEditorGraphNode* FindFeatureNodeByName(const FString& InNodeName);
+	/** Retrieves first output node that matches the InNodeName (assumes output node names are unique and includes nulloutput nodes) */
+	UE_API URigMapperDefinitionEditorGraphNode* FindOutputNodeByName(const FString& InNodeName);
 	/** Updates the key in the node map when the node is renamed (assumes new name is unique) */
 	UE_API bool UpdateNodeNameInMap(const FString& InOldName, const FString& InNewName, const ERigMapperFeatureType InNodeType);
 	/** Removes node from the map when deleted */
diff --git a/Engine/Plugins/Experimental/Animation/RigMapper/Source/RigMapperEditor/Private/RigMapperGraph/RigMapperDefinitionEditorGraphNode.cpp b/Engine/Plugins/Experimental/Animation/RigMapper/Source/RigMapperEditor/Private/RigMapperGraph/RigMapperDefinitionEditorGraphNode.cpp
--- a/Engine/Plugins/Experimental/Animation/RigMapper/Source/RigMapperEditor/Private/RigMapperGraph/RigMapperDefinitionEditorGraphNode.cpp
+++ b/Engine/Plugins/Experimental/Animation/RigMapper/Source/RigMapperEditor/Private/RigMapperGraph/RigMapperDefinitionEditorGraphNode.cpp
@@ -6,6 +6,7 @@
 #include "SRigMapperDefinitionGraphEditorNode.h"
 #include "RigMapperDefinition.h"
 #include "RigMapperDefinitionEditorGraph.h"
+#include "RigMapperUtils.h"
 #include "PropertyEditorModule.h"
 #include "EdGraph/EdGraphPin.h"
 
@@ -232,6 +233,7 @@ void URigMapperDefinitionEditorGraphNode::AddNewInputLink(const FString& InNewIn
 			Definition->Outputs[NodeName] = InNewInput;
 		}
 	}
+	UpdateNodeFields();
 }
 
 void URigMapperDefinitionEditorGraphNode::RemoveLinksToFeature(const FString& InFeatureName)
@@ -248,6 +250,14 @@ void URigMapperDefinitionEditorGraphNode::RemoveLinksToFeature(const FString& In
 	}
 
 	FeaturePtr->RenameInput(InFeatureName, TEXT(""));
+	if (NodeType == ERigMapperFeatureType::Output)
+	{
+		// Definition for Outputs has to be updated manually since RigMapperOutputFeature is created separately.
+		if (Definition->Outputs.Contains(NodeName))
+		{
+			Definition->Outputs[NodeName] = TEXT("");
+		}
+	}
 }
 
 void URigMapperDefinitionEditorGraphNode::PinConnectionListChanged(UEdGraphPin* Pin)
@@ -316,6 +326,11 @@ void URigMapperDefinitionEditorGraphNode::UpdateNodeFromFeature(const FPropertyC
 	{
 		return;
 	}
+	URigMapperDefinitionEditorGraph* RigMapperGraph = Cast<URigMapperDefinitionEditorGraph>(GetGraph());
+	if (!RigMapperGraph)
+	{
+		return;
+	}
 	FString OldName = NodeName;
 	FRigMapperFeature* FeaturePtr = nullptr;
 	if ((InPropertyChangedEvent.ChangeType & EPropertyChangeType::ValueSet) && InPropertyChangedEvent.GetPropertyName() == GET_MEMBER_NAME_CHECKED(FRigMapperFeature, Name) &&
@@ -410,7 +425,6 @@ void URigMapperDefinitionEditorGraphNode::UpdateNodeFromFeature(const FPropertyC
 				InPropertyChangedEvent.GetPropertyName() == GET_MEMBER_NAME_CHECKED(FRigMapperSdkFeature, Input))		// For SDK and Output features
 			{
 				// CHANGING AN INPUT NAME
-				URigMapperDefinitionEditorGraph* RigMapperGraph = Cast<URigMapperDefinitionEditorGraph>(GetGraph());
 				URigMapperDefinitionEditorGraphNode* LinkedNode = nullptr;
 				TArray<FString> Inputs;
 				FeaturePtr->GetInputs(Inputs);
@@ -421,7 +435,7 @@ void URigMapperDefinitionEditorGraphNode::UpdateNodeFromFeature(const FPropertyC
 					if (Inputs.IsValidIndex(InputIndex))
 					{
 						const FString InputName = Inputs[InputIndex];
-						LinkedNode = RigMapperGraph->FindNodeByName(InputName);
+						LinkedNode = RigMapperGraph->FindFeatureNodeByName(InputName);
 						if (LinkedNode)
 						{
 							InputPins[InputIndex]->MakeLinkTo(LinkedNode->GetOutputPin());
@@ -440,6 +454,20 @@ void URigMapperDefinitionEditorGraphNode::UpdateNodeFromFeature(const FPropertyC
 
 	if (NodeName != OldName)
 	{
+		bool bDuplicateNodeName = false;
+		if (NodeType == ERigMapperFeatureType::Output || NodeType == ERigMapperFeatureType::NullOutput)
+		{
+			bDuplicateNodeName = RigMapperGraph->FindOutputNodeByName(NodeName) != nullptr;						
+		}
+		else
+		{
+			bDuplicateNodeName = RigMapperGraph->FindFeatureNodeByName(NodeName) != nullptr;
+		}
+		NodeName = FRigMapperUtils::GenerateUniqueFeatureName(Definition, NodeName, NodeType);
+		// Node name must be unique.
+		// TODO: Inform user about the warning.
+		FeaturePtr->Name = NodeName;
+
 		if (NodeType != ERigMapperFeatureType::Output && NodeType != ERigMapperFeatureType::NullOutput) // Output and NullOutput nodes don't have output pins.
 		{
 			// Update inputs of all nodes directly connected to this node.
@@ -483,7 +511,6 @@ void URigMapperDefinitionEditorGraphNode::UpdateNodeFromFeature(const FPropertyC
 			}
 		}
 		// Update node map.
-		URigMapperDefinitionEditorGraph* RigMapperGraph = Cast<URigMapperDefinitionEditorGraph>(GetGraph());
 		RigMapperGraph->UpdateNodeNameInMap(OldName, NodeName, NodeType);
 	}
 		
