[RigMapper] Fixed a bug with input fields not being updated on WS and Multiply nodes

#Synced-CL 48388230

#Synced-CL 48388230

diff --git a/Engine/Plugins/Experimental/Animation/RigMapper/Source/RigMapperEditor/Private/RigMapperGraph/RigMapperDefinitionEditorGraphNode.cpp b/Engine/Plugins/Experimental/Animation/RigMapper/Source/RigMapperEditor/Private/RigMapperGraph/RigMapperDefinitionEditorGraphNode.cpp
--- a/Engine/Plugins/Experimental/Animation/RigMapper/Source/RigMapperEditor/Private/RigMapperGraph/RigMapperDefinitionEditorGraphNode.cpp
+++ b/Engine/Plugins/Experimental/Animation/RigMapper/Source/RigMapperEditor/Private/RigMapperGraph/RigMapperDefinitionEditorGraphNode.cpp
@@ -359,7 +359,6 @@ void URigMapperDefinitionEditorGraphNode::UpdateNodeFromFeature(const FPropertyC
 	NodeName = FeaturePtr->Name;
 
 	const FName PropertyName = InPropertyChangedEvent.GetPropertyName();
-	const FName ParentName = InPropertyChangedEvent.GetMemberPropertyName();
 	// TODO: Refactor with specialized node subclasses for each type
 	if (NodeType == ERigMapperFeatureType::Multiply || NodeType == ERigMapperFeatureType::WeightedSum || NodeType == ERigMapperFeatureType::SDK || NodeType == ERigMapperFeatureType::Output)
 	{
@@ -421,7 +420,7 @@ void URigMapperDefinitionEditorGraphNode::UpdateNodeFromFeature(const FPropertyC
 		}
 		else if (InPropertyChangedEvent.ChangeType & EPropertyChangeType::ValueSet)
 		{
-			if ((PinIndex != INDEX_NONE && InPropertyChangedEvent.GetPropertyName() == InputsProperty) ||				// For WS and Multiply features
+			if ((PinIndex != INDEX_NONE && InPropertyChangedEvent.GetMemberPropertyName() == InputsProperty) ||				// For WS and Multiply features
 				InPropertyChangedEvent.GetPropertyName() == GET_MEMBER_NAME_CHECKED(FRigMapperSdkFeature, Input))		// For SDK and Output features
 			{
 				// CHANGING AN INPUT NAME
