[RigMapper] Fixed a bug with uninitalized float property in RigMapperDefinition

#Synced-CL 50156828

diff --git a/Engine/Plugins/Experimental/Animation/RigMapper/Source/RigMapper/Public/RigMapperDefinition.h b/Engine/Plugins/Experimental/Animation/RigMapper/Source/RigMapper/Public/RigMapperDefinition.h
--- a/Engine/Plugins/Experimental/Animation/RigMapper/Source/RigMapper/Public/RigMapperDefinition.h
+++ b/Engine/Plugins/Experimental/Animation/RigMapper/Source/RigMapper/Public/RigMapperDefinition.h
@@ -157,7 +157,7 @@ public:
 	FString Name;
 
 	UPROPERTY(EditAnywhere, BlueprintReadWrite, Category = "Animation|Rig Mapper")
-	double Weight;
+	double Weight = 0;
 };
 
 USTRUCT(BlueprintType)
