[RigMapper] Implemented import Inputs from SkelMesh feature
Fixed a bug with RigMapper definition cache not being cleared when modified
#Synced-CL 51538485

diff --git a/Engine/Plugins/Experimental/Animation/RigMapper/Source/RigMapper/Private/RigMapperDefinition.cpp b/Engine/Plugins/Experimental/Animation/RigMapper/Source/RigMapper/Private/RigMapperDefinition.cpp
--- a/Engine/Plugins/Experimental/Animation/RigMapper/Source/RigMapper/Private/RigMapperDefinition.cpp
+++ b/Engine/Plugins/Experimental/Animation/RigMapper/Source/RigMapper/Private/RigMapperDefinition.cpp
@@ -690,6 +690,60 @@ bool URigMapperDefinition::ExportAsJsonFile(const FFilePath& JsonFilePath) const
 	return false;
 }
 
+bool URigMapperDefinition::LoadInputsFromSkeletalMesh(const TArray<FName>& InNewInputCurves)
+{
+	bool bInputAdded = false;
+	if (!InNewInputCurves.IsEmpty())
+	{
+		TSet<FName> ExistingCurves;
+		for (const FString& Input : Inputs)
+		{
+			ExistingCurves.Add(FName(Input));
+		}
+
+		for (const FName& NewInput : InNewInputCurves)
+		{
+			if (!ExistingCurves.Contains(NewInput))
+			{
+				Inputs.Add(NewInput.ToString());
+				ExistingCurves.Add(NewInput);  // keep set in sync
+				bInputAdded = true;
+			}
+		}
+	}
+	if (bInputAdded)
+	{
+		FRigMapperDefinitionsSingleton::Get().ClearFromCache(this);
+		OnRigMapperDefinitionCreated.Broadcast();
+		SetDefinitionValid(false);
+	}
+	return bInputAdded;
+}
+
+bool URigMapperDefinition::LoadOutputsFromSkeletalMesh(const TArray<FName>& InNewOutputCurves)
+{
+	bool bOutputAdded = false;
+	if (!InNewOutputCurves.IsEmpty())
+	{
+		for (const FName& NewOutput : InNewOutputCurves)
+		{
+			if (!Outputs.Contains(NewOutput.ToString()) && !NullOutputs.Contains(NewOutput.ToString()))
+			{
+				// If new output is not already used, import it, otherwise leave it as is, even if it was a NullOutput.
+				Outputs.Add(NewOutput.ToString(), TEXT(""));				
+				bOutputAdded = true;
+			}
+		}
+	}
+	if (bOutputAdded)
+	{
+		FRigMapperDefinitionsSingleton::Get().ClearFromCache(this);
+		OnRigMapperDefinitionCreated.Broadcast();
+		SetDefinitionValid(false);
+	}
+	return bOutputAdded;
+}
+
 bool URigMapperDefinition::WasDefinitionValidated() const
 {
 	return bValidated;
@@ -865,6 +919,16 @@ void URigMapperDefinition::PostEditChangeProperty(FPropertyChangedEvent& Propert
 
 	SetDefinitionValid(false);
 }
+
+void URigMapperDefinition::BroadcastDefinitionModified()
+{
+	// if the definition exists in the global cache, delete it so that the editor doesn't get out of sync with the cache
+	FRigMapperDefinitionsSingleton::Get().ClearFromCache(this);
+
+	OnRigMapperDefinitionUpdated.Broadcast();
+
+	SetDefinitionValid(false);
+}
 #endif
 
 bool URigMapperLinkedDefinitions::GetBakedInputRec(const FString& InputName, const int32 DefinitionIndex, FRigMapperFeature::FBakedInput& OutBakedInput, bool& bOutMissingIsNotNullOutput)
diff --git a/Engine/Plugins/Experimental/Animation/RigMapper/Source/RigMapper/Public/RigMapperDefinition.h b/Engine/Plugins/Experimental/Animation/RigMapper/Source/RigMapper/Public/RigMapperDefinition.h
--- a/Engine/Plugins/Experimental/Animation/RigMapper/Source/RigMapper/Public/RigMapperDefinition.h
+++ b/Engine/Plugins/Experimental/Animation/RigMapper/Source/RigMapper/Public/RigMapperDefinition.h
@@ -419,6 +419,12 @@ public:
 
 	UFUNCTION(BlueprintCallable, Category="Animation|Rig Mapper")
 	UE_API bool ExportAsJsonFile(const FFilePath& JsonFilePath) const;
+
+	UFUNCTION(BlueprintCallable, Category = "Animation|Rig Mapper")
+	UE_API bool LoadInputsFromSkeletalMesh(const TArray<FName>& InNewInputCurves);
+
+	UFUNCTION(BlueprintCallable, Category = "Animation|Rig Mapper")
+	UE_API bool LoadOutputsFromSkeletalMesh(const TArray<FName>& InNewOutputCurves);
 	
 	UFUNCTION(BlueprintCallable, Category="Animation|Rig Mapper")
 	UE_API void Empty();
@@ -449,6 +455,10 @@ private:
 #if WITH_EDITOR
 public:
 	UE_API virtual void PostEditChangeProperty(FPropertyChangedEvent& PropertyChangedEvent) override;
+	/**
+	 * Similar to PostEditChangeProperty but for Definition editor graph wrapper changes.
+	 */
+	UE_API virtual void BroadcastDefinitionModified();
 
 private:	
 	UE_API virtual EDataValidationResult IsDataValid(FDataValidationContext& Context) const override;
diff --git a/Engine/Plugins/Experimental/Animation/RigMapper/Source/RigMapperEditor/Private/RigMapperEditorModule.cpp b/Engine/Plugins/Experimental/Animation/RigMapper/Source/RigMapperEditor/Private/RigMapperEditorModule.cpp
--- a/Engine/Plugins/Experimental/Animation/RigMapper/Source/RigMapperEditor/Private/RigMapperEditorModule.cpp
+++ b/Engine/Plugins/Experimental/Animation/RigMapper/Source/RigMapperEditor/Private/RigMapperEditorModule.cpp
@@ -26,6 +26,12 @@
 #include "Misc/ConfigCacheIni.h"
 #include "Misc/Paths.h"
 #include "ScopedTransaction.h"
+#include "IContentBrowserSingleton.h"
+#include "ContentBrowserModule.h"
+#include "Engine/SkeletalMesh.h"
+#include "Animation/AnimCurveMetadata.h"
+#include "ControlRig.h"
+#include "ControlRigBlueprintLegacy.h"
 
 const FName FRigMapperEditorModule::AppIdentifier("RigMapperEditorApp");
 
@@ -159,6 +165,20 @@ void FRigMapperEditorModule::RegisterRigMapperDefinitionToolbarEntries()
 		FSlateIcon(FAppStyle::GetAppStyleSetName(), "Icons.Validate"),
 		FToolMenuExecuteAction::CreateStatic(&FRigMapperEditorModule::ValidateDefinition, true)
 	);
+	Section.AddMenuEntry(
+		"RigMapperDefinition_PopulateInputsFromMesh",
+		LOCTEXT("RigMapperDefinitionToolbar_InputImport_Label", "Populate Inputs from Mesh"),
+		LOCTEXT("RigMapperDefinitionToolbar_InputImport_Tooltip", "Pick a Skeletal Mesh and auto-create Input nodes from its curves and controls"),
+		FSlateIcon(FAppStyle::GetAppStyleSetName(), "Icons.Plus"),
+		FToolMenuExecuteAction::CreateStatic(&FRigMapperEditorModule::PopulateInputNodesFromMesh, true)
+	);
+	Section.AddMenuEntry(
+		"RigMapperDefinition_PopulateOutputsFromMesh",
+		LOCTEXT("RigMapperDefinitionToolbar_OutputImport_Label", "Populate Outputs from Mesh"),
+		LOCTEXT("RigMapperDefinitionToolbar_OutputImport_Tooltip", "Pick a Skeletal Mesh and auto-create Output nodes from its curves and controls"),
+		FSlateIcon(FAppStyle::GetAppStyleSetName(), "Icons.Plus"),
+		FToolMenuExecuteAction::CreateStatic(&FRigMapperEditorModule::PopulateOutputNodesFromMesh, true)
+	);
 }
 
 
@@ -415,6 +435,133 @@ void FRigMapperEditorModule::ValidateLinkedDefinitions(const FToolMenuContext& T
 	}
 }
 
+static void OpenImportCurvesDialogue(TArray<FName>& OutCurveNames)
+{
+	TSharedPtr<SWindow> PickerWindow;
+	bool bImportControls = false;
+	FContentBrowserModule& ContentBrowserModule =
+		FModuleManager::LoadModuleChecked<FContentBrowserModule>("ContentBrowser");
+
+	USkeletalMesh* SelectedMesh = nullptr;
+	FAssetPickerConfig AssetPickerConfig;
+	AssetPickerConfig.SelectionMode = ESelectionMode::Single;
+	AssetPickerConfig.Filter.ClassPaths.Add(USkeletalMesh::StaticClass()->GetClassPathName());
+	AssetPickerConfig.Filter.bRecursiveClasses = false;
+	AssetPickerConfig.InitialAssetViewType = EAssetViewType::List;
+	AssetPickerConfig.OnAssetSelected = FOnAssetSelected::CreateLambda(
+		[&SelectedMesh, &bImportControls, &PickerWindow](const FAssetData& AssetData)
+		{
+			SelectedMesh = Cast<USkeletalMesh>(AssetData.GetAsset());
+			bImportControls = true;
+			PickerWindow->RequestDestroyWindow();
+		});
+
+	PickerWindow = SNew(SWindow)
+		.Title(LOCTEXT("RigMapperEditorActions_PickSourceMesh", "Select Skeletal Mesh"))
+		.ClientSize(FVector2D(500, 600))
+		.SupportsMinimize(false).SupportsMaximize(false)
+		[
+			SNew(SBorder)
+				.BorderImage(FAppStyle::GetBrush("Menu.Background"))
+				[
+					ContentBrowserModule.Get().CreateAssetPicker(AssetPickerConfig)
+				]
+		];
+
+	GEditor->EditorAddModalWindow(PickerWindow.ToSharedRef());
+
+	if (SelectedMesh)
+	{
+		// Skeleton-level curves:
+		USkeleton* Skeleton = SelectedMesh->GetSkeleton();
+		if (UAnimCurveMetaData* SkeletonCurveData = Skeleton->GetAssetUserData<UAnimCurveMetaData>())
+		{
+			SkeletonCurveData->ForEachCurveMetaData(
+				[&OutCurveNames](const FName& CurveName, const FCurveMetaData& MetaData)
+				{
+					OutCurveNames.Add(CurveName);
+				});
+		}
+
+		// Mesh-level curves (additional curves not on the skeleton):
+		if (UAnimCurveMetaData* MeshCurveData = SelectedMesh->GetAssetUserData<UAnimCurveMetaData>())
+		{
+			MeshCurveData->ForEachCurveMetaData(
+				[&OutCurveNames](const FName& CurveName, const FCurveMetaData& MetaData)
+				{
+					OutCurveNames.Add(CurveName);
+				});
+		}
+
+		if (bImportControls)
+		{
+			// Control Rig-level curves:
+			if (UObject* DefaultAnimatingRig = SelectedMesh->GetDefaultAnimatingRig().LoadSynchronous())
+			{
+				if (UControlRigBlueprint* ControlRigBlueprint = Cast<UControlRigBlueprint>(DefaultAnimatingRig))
+				{
+					URigHierarchy* Hierarchy = ControlRigBlueprint->GetHierarchy();
+
+					// Iterate controls:
+					TArray<FRigControlElement*> Controls = Hierarchy->GetControls();
+					for (FRigControlElement* Control : Controls)
+					{
+						if (Control->Settings.ControlType == ERigControlType::Float)
+						{
+							FName ControlName = Control->GetFName();
+							OutCurveNames.Add(ControlName);
+						}
+					}
+
+					// Iterate curves:
+					TArray<FRigCurveElement*> Curves = Hierarchy->GetCurves();
+					for (FRigCurveElement* Curve : Curves)
+					{
+						FName CurveName = Curve->GetFName();
+						OutCurveNames.Add(CurveName);
+					}
+				}
+			}
+		}
+	}
+}
+
+void FRigMapperEditorModule::PopulateOutputNodesFromMesh(const FToolMenuContext& ToolMenuContext, bool bFromAssetEditor)
+{
+	TArray<FName> CurveNames;
+	OpenImportCurvesDialogue(CurveNames);
+
+	if (CurveNames.Num() > 0)
+	{
+		TArray<URigMapperDefinition*> Definitions = GetDefinitionsFromContext(ToolMenuContext, bFromAssetEditor);
+		FScopedTransaction Transaction(LOCTEXT("PopulateOutputNodesFromMesh", "Import output nodes from Skeletal Mesh"));
+
+		for (URigMapperDefinition* Definition : Definitions)
+		{
+			Definition->Modify();
+			Definition->LoadOutputsFromSkeletalMesh(CurveNames);
+		}
+	}
+}
+
+void FRigMapperEditorModule::PopulateInputNodesFromMesh(const FToolMenuContext& ToolMenuContext, bool bFromAssetEditor)
+{
+	TArray<FName> CurveNames;
+	OpenImportCurvesDialogue(CurveNames);
+
+	if (CurveNames.Num() > 0)
+	{
+		TArray<URigMapperDefinition*> Definitions = GetDefinitionsFromContext(ToolMenuContext, bFromAssetEditor);
+		FScopedTransaction Transaction(LOCTEXT("PopulateInputNodesFromMesh", "Import input nodes from Skeletal Mesh"));
+
+		for (URigMapperDefinition* Definition : Definitions)
+		{
+			Definition->Modify();
+			Definition->LoadInputsFromSkeletalMesh(CurveNames);
+		}
+	}
+}
+
 TArray<URigMapperLinkedDefinitions*> FRigMapperEditorModule::GetLinkedDefinitionsFromContext(const FToolMenuContext& ToolMenuContext, bool bFromAssetEditor)
 {
 	TArray<URigMapperLinkedDefinitions*> Definitions;
diff --git a/Engine/Plugins/Experimental/Animation/RigMapper/Source/RigMapperEditor/Private/RigMapperEditorModule.h b/Engine/Plugins/Experimental/Animation/RigMapper/Source/RigMapperEditor/Private/RigMapperEditorModule.h
--- a/Engine/Plugins/Experimental/Animation/RigMapper/Source/RigMapperEditor/Private/RigMapperEditorModule.h
+++ b/Engine/Plugins/Experimental/Animation/RigMapper/Source/RigMapperEditor/Private/RigMapperEditorModule.h
@@ -40,6 +40,8 @@ private:
 	static void ValidateDefinition(const FToolMenuContext& ToolMenuContext, bool bFromAssetEditor);
 	static void BakeDefinitions(const FToolMenuContext& ToolMenuContext, bool bFromAssetEditor);
 	static void ValidateLinkedDefinitions(const FToolMenuContext& ToolMenuContext, bool bFromAssetEditor);
+	static void PopulateInputNodesFromMesh(const FToolMenuContext& ToolMenuContext, bool bFromAssetEditor);
+	static void PopulateOutputNodesFromMesh(const FToolMenuContext& ToolMenuContext, bool bFromAssetEditor);
 	static TArray<URigMapperDefinition*> GetDefinitionsFromContext(const FToolMenuContext& ToolMenuContext, bool bFromAssetEditor);
 	static TArray<URigMapperLinkedDefinitions*> GetLinkedDefinitionsFromContext(const FToolMenuContext& ToolMenuContext, bool bFromAssetEditor);
 	
diff --git a/Engine/Plugins/Experimental/Animation/RigMapper/Source/RigMapperEditor/Private/RigMapperGraph/RigMapperDefinitionEditorGraph.cpp b/Engine/Plugins/Experimental/Animation/RigMapper/Source/RigMapperEditor/Private/RigMapperGraph/RigMapperDefinitionEditorGraph.cpp
--- a/Engine/Plugins/Experimental/Animation/RigMapper/Source/RigMapperEditor/Private/RigMapperGraph/RigMapperDefinitionEditorGraph.cpp
+++ b/Engine/Plugins/Experimental/Animation/RigMapper/Source/RigMapperEditor/Private/RigMapperGraph/RigMapperDefinitionEditorGraph.cpp
@@ -442,7 +442,7 @@ URigMapperDefinitionEditorGraphNode* URigMapperDefinitionEditorGraph::CreateGrap
 	}
 	NewNode->UpdateNodeFields();
 	// Notify listeners that new Feature is created.
-	WeakDefinition->OnRigMapperDefinitionUpdated.Broadcast();
+	WeakDefinition->BroadcastDefinitionModified();
 
 	return NewNode;
 }
diff --git a/Engine/Plugins/Experimental/Animation/RigMapper/Source/RigMapperEditor/Private/RigMapperGraph/RigMapperDefinitionEditorGraphNode.cpp b/Engine/Plugins/Experimental/Animation/RigMapper/Source/RigMapperEditor/Private/RigMapperGraph/RigMapperDefinitionEditorGraphNode.cpp
--- a/Engine/Plugins/Experimental/Animation/RigMapper/Source/RigMapperEditor/Private/RigMapperGraph/RigMapperDefinitionEditorGraphNode.cpp
+++ b/Engine/Plugins/Experimental/Animation/RigMapper/Source/RigMapperEditor/Private/RigMapperGraph/RigMapperDefinitionEditorGraphNode.cpp
@@ -210,6 +210,7 @@ void URigMapperDefinitionEditorGraphNode::RelinkFeature(const FString& InNewInpu
 			Definition->Outputs[NodeName] = InNewInput;
 		}
 	}
+	Definition->BroadcastDefinitionModified();
 }
 
 void URigMapperDefinitionEditorGraphNode::AddNewInputLink(const FString& InNewInput)
@@ -234,10 +235,16 @@ void URigMapperDefinitionEditorGraphNode::AddNewInputLink(const FString& InNewIn
 		}
 	}
 	UpdateNodeFields();
+	Definition->BroadcastDefinitionModified();
 }
 
 void URigMapperDefinitionEditorGraphNode::RemoveLinksToFeature(const FString& InFeatureName)
 {
+	if (!Definition.IsValid())
+	{
+		return;
+	}
+
 	if (InFeatureName.IsEmpty())
 	{
 		return;
@@ -258,6 +265,7 @@ void URigMapperDefinitionEditorGraphNode::RemoveLinksToFeature(const FString& In
 			Definition->Outputs[NodeName] = TEXT("");
 		}
 	}
+	Definition->BroadcastDefinitionModified();
 }
 
 void URigMapperDefinitionEditorGraphNode::PinConnectionListChanged(UEdGraphPin* Pin)
@@ -514,7 +522,8 @@ void URigMapperDefinitionEditorGraphNode::UpdateNodeFromFeature(const FPropertyC
 	}
 		
 	NodeTitle = FText::FromString(NodeName);
-	UpdateNodeFields();	
+	UpdateNodeFields();
+	Definition->PostEditChangeProperty(const_cast<FPropertyChangedEvent&>(InPropertyChangedEvent));
 }
 
 #undef LOCTEXT_NAMESPACE
diff --git a/Engine/Plugins/Experimental/Animation/RigMapper/Source/RigMapperEditor/Private/RigMapperGraph/SRigMapperDefinitionGraphEditor.cpp b/Engine/Plugins/Experimental/Animation/RigMapper/Source/RigMapperEditor/Private/RigMapperGraph/SRigMapperDefinitionGraphEditor.cpp
--- a/Engine/Plugins/Experimental/Animation/RigMapper/Source/RigMapperEditor/Private/RigMapperGraph/SRigMapperDefinitionGraphEditor.cpp
+++ b/Engine/Plugins/Experimental/Animation/RigMapper/Source/RigMapperEditor/Private/RigMapperGraph/SRigMapperDefinitionGraphEditor.cpp
@@ -325,7 +325,7 @@ void SRigMapperDefinitionGraphEditor::OnDeleteSelectedNodes()
 	GraphEditor->ClearSelectionSet();
 	if (URigMapperDefinitionEditorGraph* RigMapperGraph = Cast<URigMapperDefinitionEditorGraph>(CurrentGraph))
 	{
-		RigMapperGraph->GetDefinition()->OnRigMapperDefinitionUpdated.Broadcast();
+		RigMapperGraph->GetDefinition()->BroadcastDefinitionModified();
 	}
 }
 
diff --git a/Engine/Plugins/Experimental/Animation/RigMapper/Source/RigMapperEditor/RigMapperEditor.Build.cs b/Engine/Plugins/Experimental/Animation/RigMapper/Source/RigMapperEditor/RigMapperEditor.Build.cs
--- a/Engine/Plugins/Experimental/Animation/RigMapper/Source/RigMapperEditor/RigMapperEditor.Build.cs
+++ b/Engine/Plugins/Experimental/Animation/RigMapper/Source/RigMapperEditor/RigMapperEditor.Build.cs
@@ -24,6 +24,7 @@ public class RigMapperEditor : ModuleRules
 				"AnimGraph",
 				"BlueprintGraph",
 				"ControlRig",
+				"ControlRigDeveloper",
 				"LevelSequence",
 				"MovieScene",
 				"MovieSceneTracks",
