[RigMapper] RigMapper refactoring to include numerous quality of life changes:
- RigMapperNode sublasses
- Undo/Redo system
- Persistant node positions
#Synced-CL 52418140

diff --git a/Engine/Plugins/Experimental/Animation/RigMapper/Source/RigMapper/Private/RigMapperDefinition.cpp b/Engine/Plugins/Experimental/Animation/RigMapper/Source/RigMapper/Private/RigMapperDefinition.cpp
--- a/Engine/Plugins/Experimental/Animation/RigMapper/Source/RigMapper/Private/RigMapperDefinition.cpp
+++ b/Engine/Plugins/Experimental/Animation/RigMapper/Source/RigMapper/Private/RigMapperDefinition.cpp
@@ -180,19 +180,8 @@ bool FRigMapperWsFeature::IsValid(const TArray<FString>& InputNames, bool bWarn)
 }
 
 void FRigMapperWsFeature::GetInputs(TArray<FString>& OutInputs) const
-{	
-	if (!OrderedInputs.IsEmpty())
-	{
-		OutInputs.Empty(OrderedInputs.Num());
-		for (int32 i = 0; i < OrderedInputs.Num(); ++i)
-		{
-			OutInputs.Add(OrderedInputs[i].Name);
-		}
-	}
-	else
-	{
-		Inputs.GenerateKeyArray(OutInputs);
-	}
+{
+	Inputs.GenerateKeyArray(OutInputs);
 }
 
 bool FRigMapperSdkFeature::LoadFromJsonObject(const TSharedPtr<FJsonObject>& JsonObject)
@@ -908,6 +897,11 @@ bool URigMapperDefinition::Validate()
 }
 
 #if WITH_EDITOR
+void URigMapperDefinition::InvalidateCache()
+{
+	FRigMapperDefinitionsSingleton::Get().ClearFromCache(this);
+}
+
 void URigMapperDefinition::PostEditChangeProperty(FPropertyChangedEvent& PropertyChangedEvent)
 {
 	Super::PostEditChangeProperty(PropertyChangedEvent);
@@ -999,14 +993,6 @@ bool FRigMapperMultiplyFeature::BakeInput(URigMapperLinkedDefinitions* LinkedDef
 	return OutBakedInput.Key.IsValid() && !bMissingIsNotNullOutput; // if we have a missing definition which is not defined as a NullOutput, flag a failure
 }
 
-void FRigMapperMultiplyFeature::RelinkFeature(const FString& InNewInput, const int32 InputIndex)
-{
-	if (Inputs.IsValidIndex(InputIndex))
-	{
-		Inputs[InputIndex] = InNewInput;
-	}
-}
-
 void FRigMapperMultiplyFeature::AddNewInputLink(const FString& InNewInput)
 {
 	Inputs.Add(InNewInput);
@@ -1077,80 +1063,21 @@ bool FRigMapperWsFeature::BakeInput(URigMapperLinkedDefinitions* LinkedDefinitio
 	return OutBakedInput.Key.IsValid() && !bMissingIsNotNullOutput; // if we have a missing definition which is not defined as a NullOutput, flag a failure
 }
 
-void FRigMapperWsFeature::RelinkFeature(const FString& InNewInput, const int32 InputIndex)
-{
-	if (OrderedInputs.IsValidIndex(InputIndex))
-	{
-		const FString OldInput = OrderedInputs[InputIndex].Name;
-		double TempValue = 0.0f;
-		if (double* Weight = Inputs.Find(OldInput))
-		{
-			TempValue = MoveTemp(*Weight);
-			Inputs.Remove(OldInput);
-		}
-		if (!Inputs.Contains(InNewInput))
-		{
-			Inputs.Add(InNewInput, TempValue);
-		}
-		else
-		{
-			Inputs[InNewInput] += TempValue;
-		}
-		OrderedInputs[InputIndex].Name = InNewInput;
-	}
-}
-
 void FRigMapperWsFeature::AddNewInputLink(const FString& InNewInput)
 {
 	if (!Inputs.Contains(InNewInput))
 	{
 		Inputs.Add(InNewInput, 0.0f);
-		OrderedInputs.Add(FRigMapperWsInput(InNewInput, 0.0f));
-	}
-}
-
-void FRigMapperWsFeature::InitializeDetails()
-{
-	if (OrderedInputs.IsEmpty())
-	{
-		for (const TPair<FString, double>& Pair : Inputs)
-		{
-			OrderedInputs.Add(FRigMapperWsInput(Pair.Key, Pair.Value));
-		}
-	}
-}
-
-void FRigMapperWsFeature::UpdateFromDetails()
-{
-	Inputs.Empty();
-	for (FRigMapperWsInput& WsInput : OrderedInputs)
-	{
-		if (Inputs.Contains(WsInput.Name))
-		{
-			Inputs[WsInput.Name] += WsInput.Weight;
-		}
-		else
-		{
-			Inputs.Add(WsInput.Name, WsInput.Weight);
-		}
 	}
 }
 
 void FRigMapperWsFeature::RemoveInput(const FString& InInput)
 {
-	for (int32 i = OrderedInputs.Num() - 1; i >= 0; i--)
-	{
-		if (OrderedInputs[i].Name == InInput)
-		{
-			OrderedInputs.RemoveAt(i);
-		}
-	}
 	Inputs.Remove(InInput);
 }
 
 void FRigMapperWsFeature::ClearInputs()
 {
-	OrderedInputs.Empty();
 	Inputs.Empty();
 }
 
@@ -1160,13 +1087,6 @@ void FRigMapperWsFeature::RenameInput(const FString& InOldInput, const FString&
 	{
 		return;
 	}
-	for (FRigMapperWsInput& WsInput : OrderedInputs)
-	{
-		if (WsInput.Name == InOldInput)
-		{
-			WsInput.Name = InNewInput;
-		}
-	}
 	if (double* Weight = Inputs.Find(InOldInput))
 	{
 		double TempValue = MoveTemp(*Weight);
@@ -1214,11 +1134,6 @@ bool FRigMapperSdkFeature::BakeInput(URigMapperLinkedDefinitions* LinkedDefiniti
 	return OutBakedInput.Key.IsValid() && !bMissingIsNotNullOutput; // if we have a missing definition which is not defined as a NullOutput, flag a failure
 }
 
-void FRigMapperSdkFeature::RelinkFeature(const FString& InNewInput, const int32 InputIndex)
-{
-	Input = InNewInput;
-}
-
 void FRigMapperSdkFeature::AddNewInputLink(const FString& InNewInput)
 {
 	Input = InNewInput;
@@ -1501,11 +1416,6 @@ bool URigMapperLinkedDefinitions::AreLinkedDefinitionsValid() const
 	return bOk;
 }
 
-void FRigMapperOutputFeature::RelinkFeature(const FString& InNewInput, const int32 InputIndex)
-{
-	Input = InNewInput;
-}
-
 void FRigMapperOutputFeature::AddNewInputLink(const FString& InNewInput)
 {
 	Input = InNewInput;
diff --git a/Engine/Plugins/Experimental/Animation/RigMapper/Source/RigMapper/Private/RigMapperUtils.cpp b/Engine/Plugins/Experimental/Animation/RigMapper/Source/RigMapper/Private/RigMapperUtils.cpp
--- a/Engine/Plugins/Experimental/Animation/RigMapper/Source/RigMapper/Private/RigMapperUtils.cpp
+++ b/Engine/Plugins/Experimental/Animation/RigMapper/Source/RigMapper/Private/RigMapperUtils.cpp
@@ -16,30 +16,10 @@ const TMap<ERigMapperFeatureType, FString> FeatureTypePrefixMap = {
 
 const FString GenericNodeName = TEXT("node");
 
-FString FRigMapperUtils::GenerateUniqueFeatureName(const TWeakObjectPtr<URigMapperDefinition> InDefinition, FString InDesiredName, ERigMapperFeatureType InFeatureType)
+FString FRigMapperUtils::GenerateUniqueFeatureName(const TArray<FString>& InExistingNames, FString InDesiredName, ERigMapperFeatureType InFeatureType)
 {
-	TSet<FString> ExistingNames;
-	if (InFeatureType == ERigMapperFeatureType::Output || InFeatureType == ERigMapperFeatureType::NullOutput)
+	if (InFeatureType != ERigMapperFeatureType::Output && InFeatureType != ERigMapperFeatureType::NullOutput)
 	{
-		InDefinition.Pin()->Outputs.GetKeys(ExistingNames);
-		ExistingNames.Append(InDefinition.Pin()->NullOutputs);
-	}
-	else
-	{
-		ExistingNames.Append(InDefinition.Pin()->Inputs);
-		for (const FRigMapperFeature& Feature : InDefinition.Pin()->Features.Multiply)
-		{
-			ExistingNames.Add(Feature.Name);
-		}
-		for (const FRigMapperFeature& Feature : InDefinition.Pin()->Features.WeightedSums)
-		{
-			ExistingNames.Add(Feature.Name);
-		}
-		for (const FRigMapperFeature& Feature : InDefinition.Pin()->Features.SDKs)
-		{
-			ExistingNames.Add(Feature.Name);
-		}
-		
 		if (InFeatureType != ERigMapperFeatureType::Input)
 		{
 			InDesiredName = [&InDesiredName, InFeatureType]()
@@ -66,7 +46,7 @@ FString FRigMapperUtils::GenerateUniqueFeatureName(const TWeakObjectPtr<URigMapp
 	}
 
 	// If the desired name is unique, just use that.
-	if (Algo::Find(ExistingNames, InDesiredName) == nullptr)
+	if (Algo::Find(InExistingNames, InDesiredName) == nullptr)
 	{
 		return InDesiredName;
 	}
@@ -117,7 +97,7 @@ FString FRigMapperUtils::GenerateUniqueFeatureName(const TWeakObjectPtr<URigMapp
 
 	// Only test names that start with the plain name string to optimize for performance
 	TArray<FString> FilteredExistingNames;
-	Algo::TransformIf(ExistingNames, FilteredExistingNames,
+	Algo::TransformIf(InExistingNames, FilteredExistingNames,
 		[UniqueNameString = UniqueName.GetPlainNameString()](const FString& ExistingName)
 		{
 			return ExistingName.StartsWith(UniqueNameString);
diff --git a/Engine/Plugins/Experimental/Animation/RigMapper/Source/RigMapper/Public/RigMapperDefinition.h b/Engine/Plugins/Experimental/Animation/RigMapper/Source/RigMapper/Public/RigMapperDefinition.h
--- a/Engine/Plugins/Experimental/Animation/RigMapper/Source/RigMapper/Public/RigMapperDefinition.h
+++ b/Engine/Plugins/Experimental/Animation/RigMapper/Source/RigMapper/Public/RigMapperDefinition.h
@@ -14,6 +14,7 @@
 class FJsonObject;
 class FJsonValue;
 class URigMapperLinkedDefinitions;
+class UEdGraph;
 
 UENUM(BlueprintType)
 enum class ERigMapperFeatureType : uint8
@@ -55,16 +56,8 @@ public:
 
 	virtual bool BakeInput(URigMapperLinkedDefinitions* LinkedDefinitions, int32 DefinitionIndex, FBakedInput& OutBakedInput) const { return false; }
 
-	virtual void RelinkFeature(const FString& InNewInput, const int32 InputIndex) {}
-
 	virtual void AddNewInputLink(const FString& InNewInput) {}
 
-	/** Initializes additional non-serialized properties used for details view */
-	virtual void InitializeDetails() {}
-
-	/** Update feature properties from details view */
-	virtual void UpdateFromDetails() {}
-
 	virtual void RemoveInput(const FString& InInput) {};
 	virtual void ClearInputs() {};
 	virtual void RenameInput(const FString& InOldInput, const FString& InNewInput) {};
@@ -103,8 +96,6 @@ public:
 
 	UE_API virtual bool BakeInput(URigMapperLinkedDefinitions* LinkedDefinitions, int32 DefinitionIndex, FBakedInput& OutBakedInput) const override;
 
-	UE_API virtual void RelinkFeature(const FString& InNewInput, const int32 InputIndex) override;
-	
 	UE_API virtual void AddNewInputLink(const FString& InNewInput) override;
 
 	UE_API virtual void RemoveInput(const FString& InInput) override;
@@ -140,25 +131,7 @@ public:
 		return bHasLowerBound == InOther.bHasLowerBound && FMath::IsNearlyEqual(LowerBound, InOther.LowerBound, SMALL_NUMBER)
 			&& bHasUpperBound == InOther.bHasUpperBound && FMath::IsNearlyEqual(UpperBound, InOther.UpperBound, SMALL_NUMBER);
 	}
-};
-
-USTRUCT(BlueprintType)
-struct FRigMapperWsInput
-{
-	GENERATED_BODY()
-
-public:
-	FRigMapperWsInput() {}
-
-	explicit FRigMapperWsInput(const FString& InName, const double InWeight) : Name(InName), Weight(InWeight) {}
-
-public:
-	UPROPERTY(EditAnywhere, BlueprintReadWrite, Category = "Animation|Rig Mapper")
-	FString Name;
-
-	UPROPERTY(EditAnywhere, BlueprintReadWrite, Category = "Animation|Rig Mapper")
-	double Weight = 0;
-};
+};	
 
 USTRUCT(BlueprintType)
 struct FRigMapperWsFeature : public FRigMapperFeature
@@ -169,9 +142,6 @@ public:
 	UPROPERTY()
 	TMap<FString, double> Inputs;
 
-	UPROPERTY(EditAnywhere, BlueprintReadWrite, Category="Animation|Rig Mapper", meta=(DisplayPriority=1))
-	TArray<FRigMapperWsInput> OrderedInputs;
-
 	UPROPERTY(EditAnywhere, BlueprintReadWrite, Category="Animation|Rig Mapper", meta=(DisplayPriority=1))
 	FRigMapperFeatureRange Range;
 	
@@ -190,14 +160,8 @@ public:
 
 	UE_API virtual bool BakeInput(URigMapperLinkedDefinitions* LinkedDefinitions, int32 DefinitionIndex, FBakedInput& OutBakedInput) const override;
 
-	UE_API virtual void RelinkFeature(const FString& InNewInput, const int32 InputIndex) override;
-
 	UE_API virtual void AddNewInputLink(const FString& InNewInput) override;
 
-	UE_API virtual void InitializeDetails() override;
-
-	UE_API virtual void UpdateFromDetails() override;
-
 	UE_API virtual void RemoveInput(const FString& InInput) override;
 	UE_API virtual void ClearInputs() override;
 	UE_API virtual void RenameInput(const FString& InOldInput, const FString& InNewInput) override;
@@ -267,8 +231,6 @@ public:
 
 	UE_API virtual bool BakeInput(URigMapperLinkedDefinitions* LinkedDefinitions, int32 DefinitionIndex, FBakedInput& OutBakedInput) const override;
 
-	UE_API virtual void RelinkFeature(const FString& InNewInput, const int32 InputIndex) override;
-
 	UE_API virtual void AddNewInputLink(const FString& InNewInput) override;
 
 	UE_API virtual void RemoveInput(const FString& InInput) override;
@@ -301,8 +263,6 @@ public:
 
 	virtual void GetInputs(TArray<FString>& OutInputs) const override { OutInputs = { Input }; };
 
-	UE_API virtual void RelinkFeature(const FString& InNewInput, const int32 InputIndex) override;
-
 	UE_API virtual void AddNewInputLink(const FString& InNewInput) override;
 
 	UE_API virtual void RemoveInput(const FString& InInput) override;
@@ -327,8 +287,6 @@ public:
 
 	virtual ERigMapperFeatureType GetFeatureType() const override { return ERigMapperFeatureType::NullOutput; }
 
-	virtual void RelinkFeature(const FString& InNewInput, const int32 InputIndex) override {};
-
 	virtual void AddNewInputLink(const FString& InNewInput) override {};
 };
 
@@ -402,6 +360,12 @@ public:
 
 	UPROPERTY(VisibleAnywhere, BlueprintReadOnly, Category="Animation|Rig Mapper")
 	bool bValidated = false;
+
+#if WITH_EDITORONLY_DATA
+	/** Graph used by the definition editor. */
+	UPROPERTY()
+	TObjectPtr<UEdGraph> EditorGraph;
+#endif
 	
 private:
 	UPROPERTY()
@@ -454,6 +418,10 @@ private:
 
 #if WITH_EDITOR
 public:
+	/**
+	 * Only clears definition from the cache.
+	 */
+	UE_API virtual void InvalidateCache();
 	UE_API virtual void PostEditChangeProperty(FPropertyChangedEvent& PropertyChangedEvent) override;
 	/**
 	 * Similar to PostEditChangeProperty but for Definition editor graph wrapper changes.
diff --git a/Engine/Plugins/Experimental/Animation/RigMapper/Source/RigMapper/Public/RigMapperUtils.h b/Engine/Plugins/Experimental/Animation/RigMapper/Source/RigMapper/Public/RigMapperUtils.h
--- a/Engine/Plugins/Experimental/Animation/RigMapper/Source/RigMapper/Public/RigMapperUtils.h
+++ b/Engine/Plugins/Experimental/Animation/RigMapper/Source/RigMapper/Public/RigMapperUtils.h
@@ -6,21 +6,20 @@
 
 #define UE_API RIGMAPPER_API
 
-class URigMapperDefinition;
 enum class ERigMapperFeatureType : uint8;
 
 class FRigMapperUtils
 {
 public:
 	/**
-	 * Generates a unique name given a base one and a list of existing ones from the Definition, by appending a type prefix and an index suffix.
+	 * Generates a unique name given a base one and a list of existing ones, by appending a type prefix and an index suffix.
 	 *
-	 * @param InDefinition		Definition which contains all the features to compare with.
+	 * @param InExistingNames	List of existing feature names including inputs or outputs, depending on the type.
 	 * @param InDesiredName		The desired Name. Note, generates '{type name}:node' if desired name is empty.
 	 * @param InFeatureType		Type of the feature being generated to include/replace as a prefix.
 	 * @return					Returns the unique name.
 	 */
-	static UE_API FString GenerateUniqueFeatureName(const TWeakObjectPtr<URigMapperDefinition> InDefinition, FString InDesiredName, ERigMapperFeatureType InFeatureType);
+	static UE_API FString GenerateUniqueFeatureName(const TArray<FString>& InExistingNames, FString InDesiredName, ERigMapperFeatureType InFeatureType);
 };
 
 #undef UE_API
diff --git a/Engine/Plugins/Experimental/Animation/RigMapper/Source/RigMapperEditor/Private/RigMapperDefinitionEditorToolkit.cpp b/Engine/Plugins/Experimental/Animation/RigMapper/Source/RigMapperEditor/Private/RigMapperDefinitionEditorToolkit.cpp
--- a/Engine/Plugins/Experimental/Animation/RigMapper/Source/RigMapperEditor/Private/RigMapperDefinitionEditorToolkit.cpp
+++ b/Engine/Plugins/Experimental/Animation/RigMapper/Source/RigMapperEditor/Private/RigMapperDefinitionEditorToolkit.cpp
@@ -4,12 +4,12 @@
 
 #include "DetailCategoryBuilder.h"
 #include "DetailLayoutBuilder.h"
+#include "RigMapperGraph/RigMapperDefinitionEditorGraph.h"
 #include "RigMapperGraph/RigMapperDefinitionEditorGraphNode.h"
 #include "RigMapperGraph/SRigMapperDefinitionGraphEditor.h"
 #include "RigMapperEditorModule.h"
 
 #include "DetailsViewObjectFilter.h"
-#include "IStructureDetailsView.h"
 #include "PropertyEditorModule.h"
 #include "Modules/ModuleManager.h"
 #include "RigMapperGraph/SRigMapperDefinitionStructureView.h"
@@ -36,10 +36,6 @@ const TMap<FName, ERigMapperFeatureType> FRigMapperDefinitionEditorToolkit::Prop
 void FRigMapperDefinitionEditorToolkit::Initialize(URigMapperDefinition* InDefinition, EToolkitMode::Type InMode, TSharedPtr<IToolkitHost> InToolkitHost)
 {
 	Definition = InDefinition;
-	for (FRigMapperFeature& Feature : Definition->Features.WeightedSums)
-	{
-		Feature.InitializeDetails();
-	}	
 	// todo
 	// GEditor->GetEditorSubsystem<UImportSubsystem>()->OnAssetPostImport.AddSP(this, &FSimpleAssetEditor::HandleAssetPostImport);
 	// FCoreUObjectDelegates::OnObjectsReplaced.AddSP(this, &FSimpleAssetEditor::OnObjectsReplaced);
@@ -88,11 +84,24 @@ void FRigMapperDefinitionEditorToolkit::Initialize(URigMapperDefinition* InDefin
 
 void FRigMapperDefinitionEditorToolkit::OnClose()
 {
+
 	if (Definition)
 	{
+		// Sync graph state back to Definition before closing
+		if (GraphEditor)
+		{
+			if (URigMapperDefinitionEditorGraph* Graph = Cast<URigMapperDefinitionEditorGraph>(GraphEditor->GetEditorGraphObject()))
+			{
+				Graph->ApplyToDefinition();
+			}
+		}
 		Definition->OnRigMapperDefinitionCreated.RemoveAll(this);
 		Definition->OnRigMapperDefinitionUpdated.RemoveAll(this);
 	}
+	if (DetailsView)
+	{
+		DetailsView->OnFinishedChangingProperties().Clear();
+	}
 }
 
 void FRigMapperDefinitionEditorToolkit::RegisterTabSpawners(const TSharedRef<FTabManager>& InTabManager)
@@ -144,6 +153,21 @@ FLinearColor FRigMapperDefinitionEditorToolkit::GetWorldCentricTabColorScale() c
 	return FLinearColor( 0.0f, 0.1f, 0.2f, 0.5f );
 }
 
+void FRigMapperDefinitionEditorToolkit::SaveAsset_Execute()
+{
+	if (GraphEditor)
+	{
+		if (URigMapperDefinitionEditorGraph* Graph =
+			Cast<URigMapperDefinitionEditorGraph>(GraphEditor->GetEditorGraphObject()))
+		{
+			Graph->ApplyToDefinition();
+		}
+	}
+
+	// Standard save
+	FAssetEditorToolkit::SaveAsset_Execute();
+}
+
 TSharedRef<SDockTab> FRigMapperDefinitionEditorToolkit::SpawnGraphTab(const FSpawnTabArgs& SpawnTabArgs)
 {
 	// todo: search and replace (easy rename inputs, outputs and features)
@@ -177,14 +201,10 @@ void FRigMapperDefinitionEditorToolkit::HandleRigMapperDefinitionLoaded()
 	{
 		StructureView->RebuildTree();
 	}
-	if (FeatureDetailsView)
+	if (DetailsView)
 	{
-		FeatureDetailsView->SetStructureData(nullptr);
+		DetailsView->SetObject(nullptr);
 	}
-	for (FRigMapperFeature& Feature : Definition->Features.WeightedSums)
-	{
-		Feature.InitializeDetails();
-	}	
 }
 
 void FRigMapperDefinitionEditorToolkit::HandleRigMapperDefinitionUpdated()
@@ -213,30 +233,15 @@ void FRigMapperDefinitionEditorToolkit::FDetailsViewCustomization::CustomizeDeta
 TSharedRef<SDockTab> FRigMapperDefinitionEditorToolkit::SpawnDetailsTab(const FSpawnTabArgs& SpawnTabArgs)
 {
 	FPropertyEditorModule& PropertyModule = FModuleManager::LoadModuleChecked<FPropertyEditorModule>("PropertyEditor");
-	
-	FDetailsViewArgs DetailsViewArgs;
-	{
-		DetailsViewArgs.bAllowSearch = false;
-		DetailsViewArgs.NameAreaSettings = FDetailsViewArgs::HideNameArea;
-		DetailsViewArgs.bHideSelectionTip = true;
-		DetailsViewArgs.DefaultsOnlyVisibility = EEditDefaultsOnlyNodeVisibility::Automatic;
-		DetailsViewArgs.bShowOptions = false;
-		DetailsViewArgs.bAllowMultipleTopLevelObjects = true;
-		DetailsViewArgs.bShowKeyablePropertiesOption = false;
-		DetailsViewArgs.bShowModifiedPropertiesOption = false;
-		DetailsViewArgs.bAllowFavoriteSystem = false;
-		DetailsViewArgs.bShowAnimatedPropertiesOption = false;
-	}
 
-	FStructureDetailsViewArgs StructureViewArgs;
-	{
-		StructureViewArgs.bShowObjects = true;
-		StructureViewArgs.bShowAssets = true;
-		StructureViewArgs.bShowClasses = true;
-		StructureViewArgs.bShowInterfaces = true;
-	}
+	FDetailsViewArgs DetailsViewArgs;
+	DetailsViewArgs.bAllowSearch = false;
+	DetailsViewArgs.NameAreaSettings = FDetailsViewArgs::HideNameArea;
+	DetailsViewArgs.bHideSelectionTip = true;
+	DetailsViewArgs.bShowOptions = false;
+	DetailsViewArgs.bAllowMultipleTopLevelObjects = false;
 
-	FeatureDetailsView = PropertyModule.CreateStructureDetailView(DetailsViewArgs, StructureViewArgs, nullptr);
+	DetailsView = PropertyModule.CreateDetailView(DetailsViewArgs);
 
 	if (Definition)
 	{
@@ -244,12 +249,12 @@ TSharedRef<SDockTab> FRigMapperDefinitionEditorToolkit::SpawnDetailsTab(const FS
 	}
 
 	TSharedRef<SDockTab> DockTab = SNew(SDockTab)
-	.TabColorScale(GetTabColorScale())
-	.Label(LOCTEXT("RigMapperDefinitionEditorToolkitDetailsTab", "Details"))
-	[
-		FeatureDetailsView->GetWidget().ToSharedRef()
-	];
-		
+		.TabColorScale(GetTabColorScale())
+		.Label(LOCTEXT("ToolkitDetailsTab", "Details"))
+		[
+			DetailsView.ToSharedRef()
+		];
+
 	return DockTab;
 }
 
@@ -294,24 +299,30 @@ bool FRigMapperDefinitionEditorToolkit::HandleIsPropertyVisible(const FPropertyA
 
 void FRigMapperDefinitionEditorToolkit::HandleGraphSelectionChanged(const TSet<UObject*>& Nodes)
 {
-	if (!StructureView || !FeatureDetailsView)
+	if (!StructureView || !DetailsView)
 	{
 		return;
 	}
 
 	StructureView->ClearSelection();
-	FeatureDetailsView->SetStructureData(nullptr);
+	DetailsView->SetObject(nullptr);
+	SelectedGraphNode = nullptr;
 	for (const UObject* ObjectNode : Nodes)
 	{
 		if (const URigMapperDefinitionEditorGraphNode* Node = Cast<URigMapperDefinitionEditorGraphNode>(ObjectNode))
 		{
-			StructureView->SelectNode(Node->GetNodeName(), Node->GetNodeType(), true /*Is selected*/);
+			StructureView->SelectNode(Node->NodeName, Node->GetNodeType(), true);
 			if (Nodes.Num() == 1)
 			{
-				SetFeatureDetailsView(Node);
-			}					
+				SelectedGraphNode = const_cast<URigMapperDefinitionEditorGraphNode*>(Node);
+				DetailsView->SetObject(SelectedGraphNode.Get());
+				DetailsView->OnFinishedChangingProperties().Clear();
+				DetailsView->OnFinishedChangingProperties().AddRaw(
+					this, &FRigMapperDefinitionEditorToolkit::HandleNodePropertyChanged);
+			}
 		}
 	}
+
 	if (GraphTab.IsValid())
 	{
 		FGlobalTabmanager::Get()->SetActiveTab(GraphTab);
@@ -326,76 +337,51 @@ void FRigMapperDefinitionEditorToolkit::HandleGraphSelectionChanged(const TSet<U
 
 void FRigMapperDefinitionEditorToolkit::HandleStructureSelectionChanged(ESelectInfo::Type SelectInfo, TArray<FString> SelectedInputs, TArray<FString> SelectedFeatures, TArray<FString> SelectedOutputs, TArray<FString> SelectedNullOutputs)
 {
-	if (SelectInfo != ESelectInfo::Direct)
+	if (GraphEditor && SelectInfo != ESelectInfo::Direct)
 	{
 		GraphEditor->SelectNodes(SelectedInputs, SelectedFeatures, SelectedOutputs, SelectedNullOutputs);
 	};
 }
 
-void FRigMapperDefinitionEditorToolkit::SetFeatureDetailsView(const URigMapperDefinitionEditorGraphNode* InNode)
+void FRigMapperDefinitionEditorToolkit::HandleNodePropertyChanged(
+	const FPropertyChangedEvent& PropertyChangedEvent)
 {
-	if (!FeatureDetailsView)
-	{
-		return;
-	}
-
-	TObjectPtr<URigMapperDefinitionEditorGraphNode> Node = ConstCast<const URigMapperDefinitionEditorGraphNode>(InNode);
-	if (!Node)
-	{
-		return;
-	}
-	const FRigMapperFeature* FeaturePtr = Node->GetNodeFeature();
-	if (!FeaturePtr)
-	{
-		return;
-	}
-	FString CurrentFeatureName = FeaturePtr->Name;
-	TSharedPtr<FStructOnScope> StructContainer = nullptr;
-	if (Node->GetNodeType() == ERigMapperFeatureType::SDK)
-	{
-		StructContainer = MakeShared<FStructOnScope>(FRigMapperSdkFeature::StaticStruct(), reinterpret_cast<uint8*>(Node->GetNodeFeature()));
-		FeatureDetailsView->SetStructureData(StructContainer);
-	}
-	else if (Node->GetNodeType() == ERigMapperFeatureType::WeightedSum)
-	{
-		StructContainer = MakeShared<FStructOnScope>(FRigMapperWsFeature::StaticStruct(), reinterpret_cast<uint8*>(Node->GetNodeFeature()));
-	}
-	else if (Node->GetNodeType() == ERigMapperFeatureType::Multiply)
-	{
-		StructContainer = MakeShared<FStructOnScope>(FRigMapperMultiplyFeature::StaticStruct(), reinterpret_cast<uint8*>(Node->GetNodeFeature()));
-		FeatureDetailsView->SetStructureData(StructContainer);
-	}
-	else if (Node->GetNodeType() == ERigMapperFeatureType::Output)
-	{
-		StructContainer = MakeShared<FStructOnScope>(FRigMapperOutputFeature::StaticStruct(), reinterpret_cast<uint8*>(Node->GetNodeFeature()));
-		FeatureDetailsView->SetStructureData(StructContainer);
-	}
-	else if (Node->GetNodeType() == ERigMapperFeatureType::Input || Node->GetNodeType() == ERigMapperFeatureType::NullOutput)
+	// Refresh the visual appearance of the selected node
+	if (GraphEditor.IsValid() && SelectedGraphNode.IsValid())
 	{
-		StructContainer = MakeShared<FStructOnScope>(FRigMapperFeature::StaticStruct(), reinterpret_cast<uint8*>(Node->GetNodeFeature()));
-		FeatureDetailsView->SetStructureData(StructContainer);
-	}
-	FeatureDetailsView->GetOnFinishedChangingPropertiesDelegate().Clear();
-	if (StructContainer)
-	{
-		FeatureDetailsView->SetStructureData(StructContainer);
-		FeatureDetailsView->GetOnFinishedChangingPropertiesDelegate().AddRaw(this, &FRigMapperDefinitionEditorToolkit::HandleFinishedChangingFeatureProperties, TWeakObjectPtr<URigMapperDefinitionEditorGraphNode>(Node.Get()), StructContainer.ToWeakPtr());
+		// Get selected node and refresh its visual
+		GraphEditor->RefreshGraphNode(SelectedGraphNode.Get());
+		if (StructureView)
+		{
+			StructureView->RebuildTree();
+		}
 	}
 }
 
-void FRigMapperDefinitionEditorToolkit::HandleFinishedChangingFeatureProperties(const FPropertyChangedEvent& InPropertyChangedEvent, TWeakObjectPtr<URigMapperDefinitionEditorGraphNode> InNode, TWeakPtr<FStructOnScope> InStructContainer)
+void FRigMapperDefinitionEditorToolkit::PostUndo(bool bSuccess)
 {
-	if (InNode.IsValid() && GraphEditor.IsValid())
-	{		
-		InNode->UpdateNodeFromFeature(InPropertyChangedEvent, InStructContainer);
-		GraphEditor->RefreshGraphNode(InNode.Get());
-		if (StructureView)
+	if (bSuccess && GraphEditor.IsValid())
+	{
+		if (URigMapperDefinitionEditorGraph* Graph =
+			Cast<URigMapperDefinitionEditorGraph>(GraphEditor->GetEditorGraphObject()))
 		{
-			StructureView->RebuildTree();
+			Graph->RebuildNodeMaps();
+			Graph->NotifyGraphChanged();
 		}
-		SetFeatureDetailsView(InNode.Get());
-	}	
-
+		SelectedGraphNode = nullptr;
+		if (DetailsView)
+		{
+			DetailsView->SetObject(nullptr);
+		}
+		if (StructureView) 
+		{ 
+			StructureView->RebuildTree(); 
+		}
+	}
 }
 
+void FRigMapperDefinitionEditorToolkit::PostRedo(bool bSuccess)
+{
+	PostUndo(bSuccess);
+}
 #undef LOCTEXT_NAMESPACE
\ No newline at end of file
diff --git a/Engine/Plugins/Experimental/Animation/RigMapper/Source/RigMapperEditor/Private/RigMapperDefinitionEditorToolkit.h b/Engine/Plugins/Experimental/Animation/RigMapper/Source/RigMapperEditor/Private/RigMapperDefinitionEditorToolkit.h
--- a/Engine/Plugins/Experimental/Animation/RigMapper/Source/RigMapperEditor/Private/RigMapperDefinitionEditorToolkit.h
+++ b/Engine/Plugins/Experimental/Animation/RigMapper/Source/RigMapperEditor/Private/RigMapperDefinitionEditorToolkit.h
@@ -14,8 +14,6 @@ enum class ERigMapperFeatureType : uint8;
 class URigMapperDefinitionEditorGraphNode;
 class SRigMapperDefinitionStructureView;
 class SRigMapperDefinitionGraphEditor;
-class IStructureDetailsView;
-class FStructOnScope;
 
 /**
  * The toolkit for the URigMapperDefinition asset editor
@@ -33,7 +31,13 @@ public:
 	UE_API virtual FText GetBaseToolkitName() const override;
 	UE_API virtual FString GetWorldCentricTabPrefix() const override;
 	UE_API virtual FLinearColor GetWorldCentricTabColorScale() const override;
-	//~ IToolkit interface
+	UE_API virtual void SaveAsset_Execute() override;
+	//~ FAssetEditorToolkit interface
+
+	//~ FSelfRegisteringEditorUndoClient Interface
+	UE_API void PostUndo(bool bSuccess) override;
+	UE_API void PostRedo(bool bSuccess) override;
+	//~ FSelfRegisteringEditorUndoClient Interface
 
 	class FDetailsViewCustomization : public IDetailCustomization
 	{
@@ -57,13 +61,12 @@ private:
 	UE_API TSharedRef<SDockTab> SpawnStructureTab(const FSpawnTabArgs& SpawnTabArgs);
 
 	UE_API bool HandleIsPropertyVisible(const FPropertyAndParent& PropertyAndParent);
-	UE_API void HandleFinishedChangingFeatureProperties(const FPropertyChangedEvent& PropertyChangedEvent, TWeakObjectPtr<URigMapperDefinitionEditorGraphNode> InNode, TWeakPtr<FStructOnScope> InStructContainer);
+	UE_API void HandleNodePropertyChanged(const FPropertyChangedEvent& PropertyChangedEvent);
 	UE_API void HandleRigMapperDefinitionLoaded();
 	UE_API void HandleRigMapperDefinitionUpdated();
 	UE_API void HandleGraphSelectionChanged(const TSet<UObject*>& Nodes);
 	
 	UE_API void HandleStructureSelectionChanged(ESelectInfo::Type SelectInfo, TArray<FString> SelectedInputs, TArray<FString> SelectedFeatures, TArray<FString> SelectedOutputs, TArray<FString> SelectedNullOutputs);
-	void SetFeatureDetailsView(const URigMapperDefinitionEditorGraphNode* InNode);
 
 private:
 	static UE_API const FName DefinitionEditorGraphTabId;
@@ -76,9 +79,9 @@ private:
 
 	TSharedPtr<SRigMapperDefinitionGraphEditor> GraphEditor;
 	TSharedPtr<IDetailsView> DetailsView;
-	TSharedPtr<IStructureDetailsView> FeatureDetailsView;
 	TSharedPtr<SRigMapperDefinitionStructureView> StructureView;
 	TSharedPtr<SDockTab> GraphTab;
+	TWeakObjectPtr<URigMapperDefinitionEditorGraphNode> SelectedGraphNode;
 
 };
 
diff --git a/Engine/Plugins/Experimental/Animation/RigMapper/Source/RigMapperEditor/Private/RigMapperEditorModule.cpp b/Engine/Plugins/Experimental/Animation/RigMapper/Source/RigMapperEditor/Private/RigMapperEditorModule.cpp
--- a/Engine/Plugins/Experimental/Animation/RigMapper/Source/RigMapperEditor/Private/RigMapperEditorModule.cpp
+++ b/Engine/Plugins/Experimental/Animation/RigMapper/Source/RigMapperEditor/Private/RigMapperEditorModule.cpp
@@ -22,6 +22,7 @@
 #include "Toolkits/AssetEditorToolkitMenuContext.h"
 #include "ToolMenus.h"
 #include "Logging/MessageLog.h"
+#include "RigMapperGraph/RigMapperDefinitionEditorGraph.h"
 #include "RigMapperGraph/RigMapperDefinitionEditorGraphNode.h"
 #include "Misc/ConfigCacheIni.h"
 #include "Misc/Paths.h"
@@ -410,7 +411,7 @@ void FRigMapperEditorModule::BakeDefinitions(const FToolMenuContext& ToolMenuCon
 			FFormatNamedArguments Arguments;
 			Arguments.Add(TEXT("ObjectName"), FText::FromString(LinkedDefinition->GetName()));
 	
-			MessageLog.Error(FText::Format(LOCTEXT("RigMapperEditorActions_BakeFailed", "Failed to bake definitions from Linked Defintion \"{ObjectName}\". See output log for more details"), Arguments));
+			MessageLog.Error(FText::Format(LOCTEXT("RigMapperEditorActions_BakeFailed", "Failed to bake definitions from Linked Definition \"{ObjectName}\". See output log for more details"), Arguments));
 			MessageLog.Open(EMessageSeverity::Error);
 		}
 	}
@@ -429,7 +430,7 @@ void FRigMapperEditorModule::ValidateLinkedDefinitions(const FToolMenuContext& T
 			FFormatNamedArguments Arguments;
 			Arguments.Add(TEXT("ObjectName"), FText::FromString(LinkedDefinition->GetName()));
 	
-			MessageLog.Error(FText::Format(LOCTEXT("RigMapperEditorActions_ValidateLinkFailed", "Failed to validate definitions from Linked Defintion \"{ObjectName}\". See output log for more details"), Arguments));
+			MessageLog.Error(FText::Format(LOCTEXT("RigMapperEditorActions_ValidateLinkFailed", "Failed to validate definitions from Linked Definition \"{ObjectName}\". See output log for more details"), Arguments));
 			MessageLog.Open(EMessageSeverity::Error);
 		}
 	}
@@ -538,8 +539,37 @@ void FRigMapperEditorModule::PopulateOutputNodesFromMesh(const FToolMenuContext&
 
 		for (URigMapperDefinition* Definition : Definitions)
 		{
+			// Get the live graph from the definition
+			URigMapperDefinitionEditorGraph* Graph = Cast<URigMapperDefinitionEditorGraph>(Definition->EditorGraph);
 			Definition->Modify();
-			Definition->LoadOutputsFromSkeletalMesh(CurveNames);
+			if (!Graph)
+			{
+				// Create graph at this point if it does not exist.
+				Graph = NewObject<URigMapperDefinitionEditorGraph>(Definition, NAME_None, RF_Transactional);
+				Definition->EditorGraph = Graph;
+				Graph->Initialize(Definition);
+				Graph->RebuildGraph();
+				Graph->LayoutNodes();
+			}
+			else if (!Graph->GetDefinition())
+			{
+				// Existing graph is loaded from context but not yet opened in editor. Initialize it.
+				Graph->Initialize(Definition);
+				Graph->RebuildNodeMaps();
+				Graph->NotifyGraphChanged();
+			}
+
+			Graph->Modify();
+			FVector2D LastInputPos = Graph->GetNextNodeSlotPosition(ERigMapperFeatureType::Output);
+			for (const FName& CurveName : CurveNames)
+			{
+				const FString NameStr = CurveName.ToString();
+
+				if (URigMapperDefinitionEditorGraphNode* NewNode = Graph->CreateGraphNode(ERigMapperFeatureType::Output, nullptr, FVector2f(LastInputPos), false, NameStr))
+				{
+					LastInputPos = Graph->GetNextNodeSlotPositionFromNode(NewNode);
+				}
+			}
 		}
 	}
 }
@@ -556,8 +586,37 @@ void FRigMapperEditorModule::PopulateInputNodesFromMesh(const FToolMenuContext&
 
 		for (URigMapperDefinition* Definition : Definitions)
 		{
+			// Get the live graph from the definition
+			URigMapperDefinitionEditorGraph* Graph = Cast<URigMapperDefinitionEditorGraph>(Definition->EditorGraph);
 			Definition->Modify();
-			Definition->LoadInputsFromSkeletalMesh(CurveNames);
+			if (!Graph)
+			{ 
+				// Create graph at this point if it does not exist.
+				Graph = NewObject<URigMapperDefinitionEditorGraph>(Definition, NAME_None, RF_Transactional);
+				Definition->EditorGraph = Graph;
+				Graph->Initialize(Definition);
+				Graph->RebuildGraph();
+				Graph->LayoutNodes();
+			}
+			else if (!Graph->GetDefinition())
+			{
+				// Existing graph is loaded from context but not yet opened in editor. Initialize it.
+				Graph->Initialize(Definition);
+				Graph->RebuildNodeMaps();
+				Graph->NotifyGraphChanged();
+			}
+
+			Graph->Modify();
+			FVector2D LastInputPos = Graph->GetNextNodeSlotPosition(ERigMapperFeatureType::Input);
+			for (const FName& CurveName : CurveNames)
+			{
+				const FString NameStr = CurveName.ToString();
+
+				if (URigMapperDefinitionEditorGraphNode* NewNode = Graph->CreateGraphNode(ERigMapperFeatureType::Input, nullptr, FVector2f(LastInputPos), false, NameStr))
+				{
+					LastInputPos = Graph->GetNextNodeSlotPositionFromNode(NewNode);
+				}
+			}
 		}
 	}
 }
diff --git a/Engine/Plugins/Experimental/Animation/RigMapper/Source/RigMapperEditor/Private/RigMapperGraph/RigMapperDefinitionEditorGraph.cpp b/Engine/Plugins/Experimental/Animation/RigMapper/Source/RigMapperEditor/Private/RigMapperGraph/RigMapperDefinitionEditorGraph.cpp
--- a/Engine/Plugins/Experimental/Animation/RigMapper/Source/RigMapperEditor/Private/RigMapperGraph/RigMapperDefinitionEditorGraph.cpp
+++ b/Engine/Plugins/Experimental/Animation/RigMapper/Source/RigMapperEditor/Private/RigMapperGraph/RigMapperDefinitionEditorGraph.cpp
@@ -8,6 +8,20 @@
 
 #include UE_INLINE_GENERATED_CPP_BY_NAME(RigMapperDefinitionEditorGraph)
 
+static UClass* GetGraphNodeClassForFeatureType(ERigMapperFeatureType InNodeType)
+{
+	switch (InNodeType)
+	{
+	case ERigMapperFeatureType::Input:       return URigMapperDefinitionEditorGraphNode_Input::StaticClass();
+	case ERigMapperFeatureType::WeightedSum: return URigMapperDefinitionEditorGraphNode_WeightedSum::StaticClass();
+	case ERigMapperFeatureType::SDK:         return URigMapperDefinitionEditorGraphNode_SDK::StaticClass();
+	case ERigMapperFeatureType::Multiply:    return URigMapperDefinitionEditorGraphNode_Multiply::StaticClass();
+	case ERigMapperFeatureType::Output:      return URigMapperDefinitionEditorGraphNode_Output::StaticClass();
+	case ERigMapperFeatureType::NullOutput:  return URigMapperDefinitionEditorGraphNode_NullOutput::StaticClass();
+	default:                                 return URigMapperDefinitionEditorGraphNode::StaticClass();
+	}
+}
+
 URigMapperDefinitionEditorGraph::URigMapperDefinitionEditorGraph(const FObjectInitializer& ObjectInitializer) : Super(ObjectInitializer)
 {
 	Schema = URigMapperDefinitionEditorGraphSchema::StaticClass();
@@ -18,6 +32,36 @@ void URigMapperDefinitionEditorGraph::Initialize(URigMapperDefinition* InDefinit
 	WeakDefinition = InDefinition;
 }
 
+void URigMapperDefinitionEditorGraph::ApplyToDefinition()
+{
+	if (!WeakDefinition.IsValid())
+	{
+		return;
+	}
+	URigMapperDefinition* Definition = WeakDefinition.Get();
+
+	// Clear existing Definition data
+	Definition->Inputs.Empty();
+	Definition->Features.Multiply.Empty();
+	Definition->Features.WeightedSums.Empty();
+	Definition->Features.SDKs.Empty();
+	Definition->Outputs.Empty();
+	Definition->NullOutputs.Empty();
+
+	// Create definition features from graph nodes
+	for (UEdGraphNode* Node : Nodes)
+	{
+		if (URigMapperDefinitionEditorGraphNode* RigMapperNode = Cast<URigMapperDefinitionEditorGraphNode>(Node))
+		{
+			RigMapperNode->ApplyToDefinition(Definition);
+		}
+	}
+
+	// Invalidate compiled cache
+	Definition->Validate();
+	Definition->InvalidateCache();
+}
+
 void URigMapperDefinitionEditorGraph::RebuildGraph()
 {
 	RemoveAllNodes();
@@ -89,7 +133,6 @@ TArray<URigMapperDefinitionEditorGraphNode*> URigMapperDefinitionEditorGraph::Ge
 	return OutNodes;
 }
 
-// TODO: Add RigMapperDefinition node subclasses
 URigMapperDefinitionEditorGraphNode* URigMapperDefinitionEditorGraph::FindFeatureNodeByName(const FString& InNodeName)
 {
 	if (URigMapperDefinitionEditorGraphNode* FoundNode = FeatureNodes.FindRef(InNodeName))
@@ -116,7 +159,6 @@ URigMapperDefinitionEditorGraphNode* URigMapperDefinitionEditorGraph::FindOutput
 	return nullptr;
 }
 
-// TODO: Add RigMapperDefinition node subclasses
 bool URigMapperDefinitionEditorGraph::UpdateNodeNameInMap(const FString& InOldName, const FString& InNewName, const ERigMapperFeatureType InNodeType)
 {
 	switch (InNodeType)
@@ -163,10 +205,6 @@ bool URigMapperDefinitionEditorGraph::UpdateNodeNameInMap(const FString& InOldNa
 
 bool URigMapperDefinitionEditorGraph::RemoveNodeFromMap(const FString& InNodeName, const ERigMapperFeatureType InNodeType)
 {
-	if (WeakDefinition.IsValid())
-	{
-		WeakDefinition->DeleteFeature(InNodeName, InNodeType);
-	}
 	switch (InNodeType)
 	{
 	case ERigMapperFeatureType::WeightedSum:
@@ -189,6 +227,114 @@ bool URigMapperDefinitionEditorGraph::RemoveNodeFromMap(const FString& InNodeNam
 	return true;
 }
 
+void URigMapperDefinitionEditorGraph::RebuildNodeMaps()
+{
+	InputNodes.Reset();
+	FeatureNodes.Reset();
+	OutputNodes.Reset();
+	NullOutputNodes.Reset();
+	for (UEdGraphNode* Node : Nodes)
+	{
+		if (URigMapperDefinitionEditorGraphNode* RMNode =
+			Cast<URigMapperDefinitionEditorGraphNode>(Node))
+		{
+			AddNodeToMap(RMNode->GetNodeType(), RMNode->NodeName, RMNode);
+		}
+	}
+}
+
+FVector2D URigMapperDefinitionEditorGraph::GetNextNodeSlotPosition(ERigMapperFeatureType InNodeType) const
+{
+	constexpr double NodeSpacingY = 25.0;
+	constexpr double InputMarginX = 50.0;
+	constexpr double OutputMarginX = 50.0;
+	constexpr double DefaultOutputOffsetX = 400.0;
+
+	const bool bIsOutput = InNodeType != ERigMapperFeatureType::Input; // Treat all node types except input as output
+	
+	const TMap<FString, URigMapperDefinitionEditorGraphNode*>& RelevantNodes = bIsOutput ? OutputNodes : InputNodes;
+	
+	double MaxBottomY = 0.0;
+	bool bFoundAny = false;
+	
+	auto AccumulateMaxY = [&MaxBottomY, &bFoundAny](const TMap<FString, URigMapperDefinitionEditorGraphNode*>& NodeMap)
+	{
+		for (const TPair<FString, URigMapperDefinitionEditorGraphNode*>& Pair : NodeMap)
+		{
+			const double BottomY = Pair.Value->NodePosY
+				+ Pair.Value->GetDimensions().Y
+				+ Pair.Value->GetMargin().Y;
+			MaxBottomY = FMath::Max(MaxBottomY, BottomY);
+			bFoundAny = true;
+		}
+	};
+
+	AccumulateMaxY(RelevantNodes);
+	if (bIsOutput)
+	{
+		AccumulateMaxY(NullOutputNodes);
+	}
+
+	const double PosY = bFoundAny ? MaxBottomY + NodeSpacingY : 0.0;
+	double PosX = 0.0;
+	if (bIsOutput)
+	{
+		if (!OutputNodes.IsEmpty())
+		{
+			PosX = OutputNodes.CreateConstIterator().Value()->NodePosX;
+		}
+		else if (!NullOutputNodes.IsEmpty())
+		{
+			PosX = NullOutputNodes.CreateConstIterator().Value()->NodePosX;
+		}
+		else
+		{
+			// Because there are no output nodes, derrive position from the righmost input node
+			double MaxInputRight = 0.0;
+			for (const TPair<FString, URigMapperDefinitionEditorGraphNode*>& Pair : InputNodes)
+			{
+				const double RightEdge = Pair.Value->NodePosX + Pair.Value->GetDimensions().X + InputMarginX;
+				MaxInputRight = FMath::Max(MaxInputRight, RightEdge);
+			}
+			PosX = MaxInputRight > 0.0 ? MaxInputRight + OutputMarginX : DefaultOutputOffsetX;
+		}
+	}
+	return FVector2D(PosX, PosY);
+}
+
+FVector2D URigMapperDefinitionEditorGraph::GetNextNodeSlotPositionFromNode(const URigMapperDefinitionEditorGraphNode* InNode) const
+{
+	constexpr double NodeSpacingY = 25.0;
+	double PosX = 0.0;
+	double PosY = 0.0;
+	if (InNode)
+	{
+		PosX = InNode->GetNodePosX();
+		PosY = InNode->GetDimensions().Y + InNode->GetNodePosY() + NodeSpacingY;
+	}
+	return FVector2D(PosX, PosY);
+}
+
+void URigMapperDefinitionEditorGraph::AddNodeToMap(const ERigMapperFeatureType NodeType, const FString& NodeName, URigMapperDefinitionEditorGraphNode* Node)
+{
+	if (NodeType == ERigMapperFeatureType::Input)
+	{
+		InputNodes.Add(NodeName, Node);
+	}
+	else if (NodeType == ERigMapperFeatureType::Output)
+	{
+		OutputNodes.Add(NodeName, Node);
+	}
+	else if (NodeType == ERigMapperFeatureType::NullOutput)
+	{
+		NullOutputNodes.Add(NodeName, Node);
+	}
+	else
+	{
+		FeatureNodes.Add(NodeName, Node);
+	}
+}
+
 URigMapperDefinitionEditorGraphNode* URigMapperDefinitionEditorGraph::ConstructGraphNodesRec(URigMapperDefinition* Definition, const FString& NodeName, bool bIsOutputNode)
 {
 	// No need to do anything is node was already created. Output nodes should not be referenced by other nodes and should only get created once.
@@ -245,37 +391,38 @@ URigMapperDefinitionEditorGraphNode* URigMapperDefinitionEditorGraph::ConstructF
 
 	if (const FRigMapperFeature* Feature = Definition->Features.Find(NodeName, FeatureType))
 	{
-		
-		URigMapperDefinitionEditorGraphNode* Node = ConstructGraphNode(NodeName, FeatureType);
-			
-		TArray<FString> FeatureInputNames;
-		Feature->GetInputs(FeatureInputNames);
-	
-		for (const FString& FeatureInput : FeatureInputNames)
+		if (URigMapperDefinitionEditorGraphNode* Node = ConstructGraphNode(NodeName, FeatureType))
 		{
-			if (URigMapperDefinitionEditorGraphNode* LinkedNode = ConstructGraphNodesRec(Definition, FeatureInput, false))
+			Node->LoadFromFeature(Feature);
+			TArray<FString> FeatureInputNames;
+			Feature->GetInputs(FeatureInputNames);
+
+			for (const FString& FeatureInput : FeatureInputNames)
 			{
-				LinkGraphNodes(LinkedNode, Node);
+				if (URigMapperDefinitionEditorGraphNode* LinkedNode = ConstructGraphNodesRec(Definition, FeatureInput, false))
+				{
+					LinkGraphNodes(LinkedNode, Node);
+				}
+				else
+				{
+					// Feature node is saved without connected input pin.
+					UEdGraphPin* InPin = Node->CreateInputPin();
+					InPin->bHidden = false;
+				}
 			}
-			else
+
+			if (FeatureInputNames.Num() == 0)
 			{
-				// Feature node is saved without connected input pin.
-				UEdGraphPin* InPin = Node->CreateInputPin();
-				InPin->bHidden = false;
+				Node->CreateInputPin();
+			}
+			if (FeatureInputNames.Num() <= 1 && FeatureType == ERigMapperFeatureType::Multiply)
+			{
+				Node->CreateInputPin();
 			}
-		}
 
-		if (FeatureInputNames.Num() == 0)
-		{
-			Node->CreateInputPin();
+			Node->UpdateNodeFields();
+			return Node;
 		}
-		if (FeatureInputNames.Num() <= 1 && FeatureType == ERigMapperFeatureType::Multiply)
-		{
-			Node->CreateInputPin();
-		}
-		
-		Node->UpdateNodeFields();
-		return Node;
 	}
 	return nullptr;
 }
@@ -310,11 +457,13 @@ URigMapperDefinitionEditorGraphNode* URigMapperDefinitionEditorGraph::ConstructG
 	
 	constexpr bool bSelectNewNode = false;
 	
-	FGraphNodeCreator<URigMapperDefinitionEditorGraphNode> GraphNodeCreator(*this);
-	URigMapperDefinitionEditorGraphNode* Node = GraphNodeCreator.CreateNode(bSelectNewNode);
-	GraphNodeCreator.Finalize();
+	URigMapperDefinitionEditorGraphNode* Node = NewObject<URigMapperDefinitionEditorGraphNode>(
+		this, GetGraphNodeClassForFeatureType(NodeType), NAME_None, RF_Transactional);
+	this->AddNode(Node, /*bUserAction=*/ false, /*bSelectNewNode=*/ bSelectNewNode);
+	Node->CreateNewGuid();
+	Node->PostPlacedNewNode();
 
-	Node->SetupNode(WeakDefinition.Get(), NodeName, NodeType);
+	Node->SetupNode(NodeName, NodeType);
 	if (NodeType != ERigMapperFeatureType::NullOutput && NodeType != ERigMapperFeatureType::Output)
 	{
 		// One output node must be created for non output types, even if node is "hanging".
@@ -331,30 +480,74 @@ URigMapperDefinitionEditorGraphNode* URigMapperDefinitionEditorGraph::ConstructG
 }
 
 
-URigMapperDefinitionEditorGraphNode* URigMapperDefinitionEditorGraph::CreateGraphNode(const ERigMapperFeatureType InNodeType, UEdGraphPin* InFromPin, const FVector2f& InLocation, bool bSelectNewNode /*= true*/)
+void URigMapperDefinitionEditorGraph::GetExistingNodeNames(const ERigMapperFeatureType InNodeType, TArray<FString>& OutExistingNames)
+{
+	OutExistingNames.Reset();
+	if (InNodeType == ERigMapperFeatureType::Output || InNodeType == ERigMapperFeatureType::NullOutput)
+	{
+		OutExistingNames.Reserve(OutputNodes.Num() + NullOutputNodes.Num());
+		for (const TPair<FString, URigMapperDefinitionEditorGraphNode*>& Pair : OutputNodes) { OutExistingNames.Add(Pair.Key); }
+		for (const TPair<FString, URigMapperDefinitionEditorGraphNode*>& Pair : NullOutputNodes) { OutExistingNames.Add(Pair.Key); }
+	}
+	else
+	{
+		OutExistingNames.Reserve(InputNodes.Num() + FeatureNodes.Num());
+		for (const TPair<FString, URigMapperDefinitionEditorGraphNode*>& Pair : InputNodes) { OutExistingNames.Add(Pair.Key); }
+		for (const TPair<FString, URigMapperDefinitionEditorGraphNode*>& Pair : FeatureNodes) { OutExistingNames.Add(Pair.Key); }
+	}
+}
+
+URigMapperDefinitionEditorGraphNode* URigMapperDefinitionEditorGraph::CreateGraphNode(const ERigMapperFeatureType InNodeType, UEdGraphPin* InFromPin, const FVector2f& InLocation, bool bSelectNewNode /*= true*/, const FString InDesiredNodeName)
 {
 	if (!WeakDefinition.IsValid())
 	{
 		return nullptr;
 	}
+
+	if (!InDesiredNodeName.IsEmpty())
+	{
+		// DesiredNodeName is expected to be unique. If it is not, no node should be created.
+		if (InNodeType == ERigMapperFeatureType::Output || InNodeType == ERigMapperFeatureType::NullOutput)
+		{
+			if (FindOutputNodeByName(InDesiredNodeName))
+			{
+				return nullptr;
+			}
+		}
+		else
+		{
+			if (FindFeatureNodeByName(InDesiredNodeName))
+			{
+				return nullptr;
+			}
+		}
+	}
 	FString FromPinName = TEXT("");
-	if (InFromPin)
+	if (InFromPin && InFromPin->GetOwningNode())
 	{
-		FromPinName = Cast<URigMapperDefinitionEditorGraphNode>(InFromPin->GetOwningNode())->GetNodeName();
+		FromPinName = Cast<URigMapperDefinitionEditorGraphNode>(InFromPin->GetOwningNode())->NodeName;
 	}
-	FGraphNodeCreator<URigMapperDefinitionEditorGraphNode> GraphNodeCreator(*this);
-	URigMapperDefinitionEditorGraphNode* NewNode = GraphNodeCreator.CreateNode(bSelectNewNode);
-	GraphNodeCreator.Finalize();
+	URigMapperDefinitionEditorGraphNode* NewNode = NewObject<URigMapperDefinitionEditorGraphNode>(
+		this, GetGraphNodeClassForFeatureType(InNodeType), NAME_None, RF_Transactional);
+	this->AddNode(NewNode, /*bUserAction=*/ false, /*bSelectNewNode=*/ bSelectNewNode);
+	NewNode->CreateNewGuid();
+	NewNode->PostPlacedNewNode();
 	NewNode->NodePosX = InLocation.X;
 	NewNode->NodePosY = InLocation.Y;
 	// Create unique feature (node) name.
-	FString NodeName = TEXT("");
-	if (InNodeType != ERigMapperFeatureType::Input && InNodeType != ERigMapperFeatureType::Output && InNodeType != ERigMapperFeatureType::NullOutput)
+	FString NodeName = InDesiredNodeName;
+	if (NodeName == TEXT(""))
 	{
-		// Only feature types can propagate names in chains.
-		NodeName = FromPinName;
+		// Generate new unique node name from pin (if provided).
+		if (InNodeType != ERigMapperFeatureType::Input && InNodeType != ERigMapperFeatureType::Output && InNodeType != ERigMapperFeatureType::NullOutput)
+		{
+			// Only feature types can propagate names in chains.
+			NodeName = FromPinName;
+		}
+		TArray<FString> ExistingNames;
+		GetExistingNodeNames(InNodeType, ExistingNames);
+		NodeName = FRigMapperUtils::GenerateUniqueFeatureName(ExistingNames, NodeName, InNodeType);
 	}
-	NodeName = FRigMapperUtils::GenerateUniqueFeatureName(WeakDefinition, NodeName, InNodeType);
 	if (InNodeType != ERigMapperFeatureType::Input && InNodeType != ERigMapperFeatureType::NullOutput)
 	{
 		UEdGraphPin* InPin = NewNode->CreateInputPin();
@@ -365,62 +558,30 @@ URigMapperDefinitionEditorGraphNode* URigMapperDefinitionEditorGraph::CreateGrap
 	{
 		InputNodeName = FromPinName;
 	}
-	if (InNodeType == ERigMapperFeatureType::Multiply)
-	{
-		FRigMapperMultiplyFeature NewFeature = FRigMapperMultiplyFeature(NodeName);
-		if (!InputNodeName.IsEmpty())
-		{
-			NewFeature.Inputs.Add(InputNodeName);
-		}
-		WeakDefinition->Features.Multiply.Add(NewFeature);		
-		// Multiply and Weighted Sum features requires at least 2 inputs.
-		UEdGraphPin* InPin = NewNode->CreateInputPin();
-		InPin->bHidden = false;
-	}
-	else if (InNodeType == ERigMapperFeatureType::WeightedSum)
+	if (InNodeType != ERigMapperFeatureType::Output && InNodeType != ERigMapperFeatureType::NullOutput)
 	{
-		FRigMapperWsFeature NewFeature = FRigMapperWsFeature(NodeName);
-		if (!InputNodeName.IsEmpty())
-		{
-			NewFeature.OrderedInputs.Add(FRigMapperWsInput(InputNodeName, 0.0));
-			NewFeature.Inputs.Add(InputNodeName);
-		}
-		WeakDefinition->Features.WeightedSums.Add(NewFeature);
+		UEdGraphPin* OutPin = NewNode->CreateOutputPin();
+		OutPin->bHidden = false;
 	}
-	else if (InNodeType == ERigMapperFeatureType::SDK)
+	if (InNodeType == ERigMapperFeatureType::WeightedSum)
 	{
-		FRigMapperSdkFeature NewFeature = FRigMapperSdkFeature(NodeName);
-		if (!InputNodeName.IsEmpty())
+		if (URigMapperDefinitionEditorGraphNode_WeightedSum* WsNode = Cast<URigMapperDefinitionEditorGraphNode_WeightedSum>(NewNode))
 		{
-			NewFeature.Input = InputNodeName;
+			WsNode->WeightedInputs.Add(FRigMapperWsInput(InputNodeName, 0.0));
 		}
-		WeakDefinition->Features.SDKs.Add(NewFeature);
-	}
-	else if (InNodeType == ERigMapperFeatureType::Output)
-	{
-		WeakDefinition->Outputs.Add(NodeName, InputNodeName);
 	}
-	else if (InNodeType == ERigMapperFeatureType::Input)
+	else if (InNodeType == ERigMapperFeatureType::Multiply)
 	{
-		WeakDefinition->Inputs.Add(NodeName);
-	}
-	else if (InNodeType == ERigMapperFeatureType::NullOutput)
-	{
-		WeakDefinition->NullOutputs.Add(NodeName);
-	}
-	if (InNodeType != ERigMapperFeatureType::Output && InNodeType != ERigMapperFeatureType::NullOutput)
-	{
-		UEdGraphPin* OutPin = NewNode->CreateOutputPin();
-		OutPin->bHidden = false;
+		// Multiply features require at least 2 inputs.
+		UEdGraphPin* InPin = NewNode->CreateInputPin();
+		InPin->bHidden = false;
 	}
-
-	NewNode->SetupNode(WeakDefinition.Get(), NodeName, InNodeType);
+	NewNode->SetupNode(NodeName, InNodeType);
 	AddNodeToMap(InNodeType, NodeName, NewNode);
 
 	if (InFromPin)
 	{
 		const UEdGraphSchema* GraphSchema = GetSchema();
-		UEdGraphPin* TargetPin = nullptr;
 		if (InFromPin->Direction == EGPD_Output)
 		{
 			const TArray<UEdGraphPin*>& Pins = NewNode->GetInputPins();
@@ -441,9 +602,8 @@ URigMapperDefinitionEditorGraphNode* URigMapperDefinitionEditorGraph::CreateGrap
 		}
 	}
 	NewNode->UpdateNodeFields();
-	// Notify listeners that new Feature is created.
-	WeakDefinition->BroadcastDefinitionModified();
-
+	// Notify listeners that new node is created.
+	OnGraphStructureUpdated.Broadcast();
 	return NewNode;
 }
 
@@ -510,27 +670,6 @@ void URigMapperDefinitionEditorGraph::LayoutNodeRec(URigMapperDefinitionEditorGr
 	InNode->SetMargin(Margin);
 }
 
-// TODO: Add RigMapperDefinition node subclasses
-void URigMapperDefinitionEditorGraph::AddNodeToMap(const ERigMapperFeatureType NodeType, const FString& NodeName, URigMapperDefinitionEditorGraphNode* Node)
-{
-	if (NodeType == ERigMapperFeatureType::Input)
-	{
-		InputNodes.Add(NodeName, Node);
-	}
-	else if (NodeType == ERigMapperFeatureType::Output)
-	{
-		OutputNodes.Add(NodeName, Node);
-	}
-	else if (NodeType == ERigMapperFeatureType::NullOutput)
-	{
-		NullOutputNodes.Add(NodeName, Node);
-	}
-	else // todo: Add feature type
-	{
-		FeatureNodes.Add(NodeName, Node);
-	}
-}
-
 void URigMapperDefinitionEditorGraph::LayoutNodes() const
 {
 	double InputsMaxWidth = 0;
@@ -567,18 +706,18 @@ void URigMapperDefinitionEditorGraph::LayoutNodes() const
 	for (const TPair<FString, URigMapperDefinitionEditorGraphNode*>& Output : OutputNodes)
 	{
 		Output.Value->NodePosX = MaxPosX + OutputMarginX;
-		// for (UEdGraphPin* InPin : Output.Value->GetInputPins())
-		// {
-		// 	for (const UEdGraphPin* OutPin : InPin->LinkedTo)
-		// 	{
-		// 		URigMapperDefinitionEditorGraphNode* LinkedNode = Cast<URigMapperDefinitionEditorGraphNode>(OutPin->GetOwningNode());
-		//
-		// 		if (!LinkedNode->GetInputPins().IsEmpty())
-		// 		{
-		// 			LinkedNode->NodePosX = Output.Value->NodePosX - (LinkedNode->GetDimensions().X + LinkedNode->GetMargin().X + OutputMarginX);
-		// 		}
-		// 	}
-		// }
+		 for (UEdGraphPin* InPin : Output.Value->GetInputPins())
+		 {
+		 	for (const UEdGraphPin* OutPin : InPin->LinkedTo)
+		 	{
+		 		URigMapperDefinitionEditorGraphNode* LinkedNode = Cast<URigMapperDefinitionEditorGraphNode>(OutPin->GetOwningNode());
+				
+		 		if (LinkedNode && !LinkedNode->GetInputPins().IsEmpty())
+		 		{
+		 			LinkedNode->NodePosX = Output.Value->NodePosX - (LinkedNode->GetDimensions().X + LinkedNode->GetMargin().X + OutputMarginX);
+		 		}
+		 	}
+		 }
 	}
 
 	// Layout nodes not related to any output
diff --git a/Engine/Plugins/Experimental/Animation/RigMapper/Source/RigMapperEditor/Private/RigMapperGraph/RigMapperDefinitionEditorGraph.h b/Engine/Plugins/Experimental/Animation/RigMapper/Source/RigMapperEditor/Private/RigMapperGraph/RigMapperDefinitionEditorGraph.h
--- a/Engine/Plugins/Experimental/Animation/RigMapper/Source/RigMapperEditor/Private/RigMapperGraph/RigMapperDefinitionEditorGraph.h
+++ b/Engine/Plugins/Experimental/Animation/RigMapper/Source/RigMapperEditor/Private/RigMapperGraph/RigMapperDefinitionEditorGraph.h
@@ -12,6 +12,8 @@ enum class ERigMapperFeatureType : uint8;
 class URigMapperDefinitionEditorGraphNode;
 class URigMapperDefinition;
 
+DECLARE_MULTICAST_DELEGATE(FOnGraphStructureUpdated);
+
 /**
  * 
  */
@@ -23,16 +25,24 @@ class URigMapperDefinitionEditorGraph : public UEdGraph
 public:
 	UE_API void Initialize(URigMapperDefinition* InDefinition);
 	
+	FOnGraphStructureUpdated OnGraphStructureUpdated;
+
+	/**
+	 * Applies graph nodes state to definition, validates it and invalidates cache.
+	 * Intended to call only on save and close so it does not call Modify() in order not to dirty the state of Definition again.
+	 * If needed elsewhere, in order to preserve the undo stack, call Modify outside this method.
+	 */
+	UE_API void ApplyToDefinition();
 	UE_API void RebuildGraph();
-	UE_API void ConstructNodes();
-	UE_API void RemoveAllNodes();
 	void RequestRefreshLayout(bool bInRefreshLayout) { bRefreshLayout = bInRefreshLayout; }
 	bool NeedsRefreshLayout() const { return bRefreshLayout; }
 	UE_API void LayoutNodes() const;
 	
 	URigMapperDefinition* GetDefinition() const { return WeakDefinition.Get(); };
 
-	UE_API URigMapperDefinitionEditorGraphNode* CreateGraphNode(const ERigMapperFeatureType InNodeType, UEdGraphPin* InFromPin, const FVector2f& InLocation, bool bSelectNewNode = true);
+	void GetExistingNodeNames(const ERigMapperFeatureType InNodeType, TArray<FString>& OutExistingNames);
+
+	UE_API URigMapperDefinitionEditorGraphNode* CreateGraphNode(const ERigMapperFeatureType InNodeType, UEdGraphPin* InFromPin, const FVector2f& InLocation, bool bSelectNewNode = true, const FString InDesiredNodeName = TEXT(""));
 	UE_API TArray<URigMapperDefinitionEditorGraphNode*> GetNodesByName(const TArray<FString>& Inputs, const TArray<FString>& Features, const TArray<FString>& Outputs, const TArray<FString>& NullOutputs) const;
 	/** Retrieves first feature node that matches the InNodeName (assumes node names are unique and includes input nodes) */
 	UE_API URigMapperDefinitionEditorGraphNode* FindFeatureNodeByName(const FString& InNodeName);
@@ -42,8 +52,16 @@ public:
 	UE_API bool UpdateNodeNameInMap(const FString& InOldName, const FString& InNewName, const ERigMapperFeatureType InNodeType);
 	/** Removes node from the map when deleted */
 	UE_API bool RemoveNodeFromMap(const FString& InNodeName, const ERigMapperFeatureType InNodeType);
+	/** Rebuilds node map from graph nodes **/
+	UE_API void RebuildNodeMaps();
+	/** Gets the next available node position. Input types are placed in the input column; all other types in the output column. */
+	UE_API FVector2D GetNextNodeSlotPosition(ERigMapperFeatureType InNodeType) const;
+	/** Gets the position for the next node below the given node **/
+	UE_API FVector2D GetNextNodeSlotPositionFromNode(const URigMapperDefinitionEditorGraphNode* InNode) const;
 
 private:
+	UE_API void ConstructNodes();
+	UE_API void RemoveAllNodes();
 	UE_API URigMapperDefinitionEditorGraphNode* ConstructGraphNode(const FString& NodeName, const ERigMapperFeatureType NodeType);
 	UE_API URigMapperDefinitionEditorGraphNode* ConstructGraphNodesRec(URigMapperDefinition* Definition, const FString& InNodeName, bool bIsOutputNode);
 	UE_API URigMapperDefinitionEditorGraphNode* ConstructOutputNode(URigMapperDefinition* Definition, const FString& NodeName);
diff --git a/Engine/Plugins/Experimental/Animation/RigMapper/Source/RigMapperEditor/Private/RigMapperGraph/RigMapperDefinitionEditorGraphNode.cpp b/Engine/Plugins/Experimental/Animation/RigMapper/Source/RigMapperEditor/Private/RigMapperGraph/RigMapperDefinitionEditorGraphNode.cpp
--- a/Engine/Plugins/Experimental/Animation/RigMapper/Source/RigMapperEditor/Private/RigMapperGraph/RigMapperDefinitionEditorGraphNode.cpp
+++ b/Engine/Plugins/Experimental/Animation/RigMapper/Source/RigMapperEditor/Private/RigMapperGraph/RigMapperDefinitionEditorGraphNode.cpp
@@ -4,7 +4,6 @@
 #include "RigMapperDefinitionEditorGraphNode.h"
 
 #include "SRigMapperDefinitionGraphEditorNode.h"
-#include "RigMapperDefinition.h"
 #include "RigMapperDefinitionEditorGraph.h"
 #include "RigMapperUtils.h"
 #include "PropertyEditorModule.h"
@@ -31,40 +30,6 @@ TSharedPtr<SGraphNode> URigMapperDefinitionEditorGraphNode::NodeFactory::CreateN
 	return nullptr;
 }
 
-FLinearColor URigMapperDefinitionEditorGraphNode::GetNodeBodyTintColor() const
-{
-	return GetNodeTitleColor();
-}
-
-FLinearColor URigMapperDefinitionEditorGraphNode::GetNodeTitleColor() const
-{
-	if (NodeType == ERigMapperFeatureType::Input)
-	{
-		return FLinearColor(0.466, 0.969, 0.878);
-	}
-	if (NodeType == ERigMapperFeatureType::Output)
-	{
-		return FLinearColor(0.466, 0.969, 0.525);
-	}
-	if (NodeType == ERigMapperFeatureType::NullOutput)
-	{
-		return FLinearColor(0.65, 0.969, 0.466);
-	}
-	if (NodeType == ERigMapperFeatureType::Multiply)
-	{
-		return FLinearColor(0.969, 0.966, 0.466);
-	}
-	if (NodeType == ERigMapperFeatureType::SDK)
-	{
-		return FLinearColor(0.969, 0.717, 0.466);
-	}
-	if (NodeType == ERigMapperFeatureType::WeightedSum)
-	{
-		return FLinearColor(0.65, 0.466, 0.969);
-	}
-	return FLinearColor(0.969, 0.5, 0.466);
-}
-
 void URigMapperDefinitionEditorGraphNode::GetNodeContextMenuActions(class UToolMenu* Menu, class UGraphNodeContextMenuContext* Context) const
 {
 
@@ -75,225 +40,250 @@ void URigMapperDefinitionEditorGraphNode::AllocateDefaultPins()
 
 }
 
-void URigMapperDefinitionEditorGraphNode::SetupNode(URigMapperDefinition* InDefinition, const FString& InNodeName, ERigMapperFeatureType InNodeType)
+void URigMapperDefinitionEditorGraphNode::SetupNode(const FString& InNodeName, ERigMapperFeatureType InNodeType)
 {
 	NodeName = InNodeName;
-	Definition = InDefinition;
 	NodeType = InNodeType;
 	NodeTitle = FText::FromString(NodeName);	
-	if (!Definition.IsValid())
-	{
-		return;
-	}
+}
 
-	// Create features not tracked with Definition and assign the Definition as the owner.
-	if (NodeType == ERigMapperFeatureType::Output)
-	{
-		FRigMapperOutputFeature NewFeature = FRigMapperOutputFeature(NodeName);
-		NewFeature.Input = Definition->Outputs[NodeName];
-		Definition->Features.Outputs.Add(NewFeature);
-	}
-	else if (NodeType == ERigMapperFeatureType::Input)
-	{
-		FRigMapperFeature NewFeature = FRigMapperFeature(NodeName);
-		Definition->Features.Inputs.Add(NewFeature);
-	}
-	else if (NodeType == ERigMapperFeatureType::NullOutput)
+
+void URigMapperDefinitionEditorGraphNode::PostLoad()
+{
+	Super::PostLoad();
+
+	InputPins.Reset();
+	OutputPins.Reset();
+	for (UEdGraphPin* Pin : Pins)
 	{
-		FRigMapperNullOutputFeature NewFeature = FRigMapperNullOutputFeature(NodeName);
-		Definition->Features.NullOutputs.Add(NewFeature);
+		if (Pin->Direction == EGPD_Input) { InputPins.Add(Pin); }
+		else if (Pin->Direction == EGPD_Output) { OutputPins.Add(Pin); }
 	}
+	UpdateNodeFields();
 }
 
 void URigMapperDefinitionEditorGraphNode::UpdateNodeFields()
 {
-	if (!Definition.IsValid())
-	{
-		return;
-	}
+	NodeTitle = FText::FromString(NodeName);
+}
 
-	FRigMapperFeature* FeaturePtr = GetNodeFeature();
-	if (!FeaturePtr)
+void URigMapperDefinitionEditorGraphNode::RelinkFeature(const FString& InNewInput, UEdGraphPin* InputPin)
+{
+	if (!InputPin)
 	{
 		return;
 	}
 
-	if (NodeType == ERigMapperFeatureType::WeightedSum || NodeType == ERigMapperFeatureType::SDK)
-	{
-		if (FeaturePtr->GetFeatureType() == ERigMapperFeatureType::SDK)
-		{
-			TArray<FString> InKeys;
-			TArray<FString> OutKeys;
-
-			const TArray<FRigMapperSdkKey>& Keys = static_cast<FRigMapperSdkFeature*>(FeaturePtr)->Keys;
+	const int32 PinIndex = InputPins.Find(InputPin);
+	if (PinIndex == INDEX_NONE) return;
 
-			InKeys.Reserve(Keys.Num());
-			OutKeys.Reserve(Keys.Num());
-			for (const FRigMapperSdkKey& Key : Keys)
-			{
-				InKeys.Add(FString::Printf(TEXT("%.*f"), 3, Key.In));
-				OutKeys.Add(FString::Printf(TEXT("%.*f"), 3, Key.Out));
-			}
-			const FString In = FString::Join(InKeys, TEXT(", "));
-			const FString Out = FString::Join(OutKeys, TEXT(", "));
-			NodeSubtitle = FText::FromString(FString::Printf(TEXT("[%s] > [%s]"), *In, *Out));
-		}
-		if (FeaturePtr->GetFeatureType() == ERigMapperFeatureType::WeightedSum)
+	Modify();
+	if (NodeType == ERigMapperFeatureType::WeightedSum)
+	{
+		if (URigMapperDefinitionEditorGraphNode_WeightedSum* WsNode =
+			Cast<URigMapperDefinitionEditorGraphNode_WeightedSum>(this))
 		{
-			const FRigMapperWsFeature* Feature = static_cast<FRigMapperWsFeature*>(FeaturePtr);
-			const TArray<FRigMapperWsInput> Inputs = Feature->OrderedInputs;
-			const FRigMapperFeatureRange& Range = Feature->Range;
-			const FText MinText = LOCTEXT("RangeMinimum", "min");
-			const FText MaxText = LOCTEXT("RangeMaximum", "max");
-
-			for (int32 PinIndex = 0; PinIndex < InputPins.Num(); PinIndex++)
-			{
-				if (Inputs.IsValidIndex(PinIndex))
-				{
-					InputPins[PinIndex]->PinFriendlyName = FText::Format(INVTEXT("{0}"), Inputs[PinIndex].Weight);
-				}
-			}
-			FString RangeStr;
-			if (Range.bHasLowerBound)
+			if (WsNode->WeightedInputs.IsValidIndex(PinIndex))
 			{
-				RangeStr.Appendf(TEXT("%s: %.*f"), *MinText.ToString(), 3, Range.LowerBound);
+				WsNode->WeightedInputs[PinIndex].Name = InNewInput;
 			}
-			if (Range.bHasLowerBound && Range.bHasUpperBound)
+			else
 			{
-				RangeStr.Append(TEXT("\n"));
+				// Pin exists but no WeightedInput entry, so one should be added
+				WsNode->WeightedInputs.Add(FRigMapperWsInput(InNewInput, 0.0));
 			}
-			if (Range.bHasUpperBound)
-			{
-				RangeStr.Appendf(TEXT("%s: %.*f"), *MaxText.ToString(), 3, Range.UpperBound);
-			}
-			NodeSubtitle = FText::FromString(RangeStr);
 		}
 	}
+	UpdateNodeFields();
 }
 
-void URigMapperDefinitionEditorGraphNode::RelinkFeature(const FString& InNewInput, UEdGraphPin* InputPin)
+void URigMapperDefinitionEditorGraphNode::AddNewInputLink(const FString& InNewInput)
 {
-	if (!Definition.IsValid())
-	{
-		return;
-	}
-	FRigMapperFeature* FeaturePtr = GetNodeFeature();
-	if (!FeaturePtr)
-	{
-		return;
-	}
-	if (!InputPin)
-	{
-		return;
-	}
-	TArray<FString> InputNames;
-	FeaturePtr->GetInputs(InputNames);
-	int32 PinIndex = InputPins.Find(InputPin);
-	if (PinIndex == INDEX_NONE)
-	{
-		return;
-	}
-	if (InputNames.Num() <= PinIndex)
-	{
-		// New input should be added to definition.
-		FeaturePtr->AddNewInputLink(InNewInput);
-	}
-	else
-	{
-		FeaturePtr->RelinkFeature(InNewInput, PinIndex);
-	}
-	if (NodeType == ERigMapperFeatureType::Output)
+	Modify();
+	if (NodeType == ERigMapperFeatureType::WeightedSum)
 	{
-		// Definition for Outputs has to be updated manually since RigMapperOutputFeature is created separately.
-		if (Definition->Outputs.Contains(NodeName))
+		if (URigMapperDefinitionEditorGraphNode_WeightedSum* WsNode =
+			Cast<URigMapperDefinitionEditorGraphNode_WeightedSum>(this))
 		{
-			Definition->Outputs[NodeName] = InNewInput;
+			WsNode->WeightedInputs.Add(FRigMapperWsInput(InNewInput, 0.0));
 		}
 	}
-	Definition->BroadcastDefinitionModified();
+	UpdateNodeFields();
 }
 
-void URigMapperDefinitionEditorGraphNode::AddNewInputLink(const FString& InNewInput)
+void URigMapperDefinitionEditorGraphNode::RemoveLinksToFeature(const FString& InFeatureName)
 {
-	if (!Definition.IsValid())
-	{
-		return;
-	}
-
-	FRigMapperFeature* FeaturePtr = GetNodeFeature();
-	if (!FeaturePtr)
-	{
-		return;
-	}
-	FeaturePtr->AddNewInputLink(InNewInput);
-	if (NodeType == ERigMapperFeatureType::Output)
+	Modify();
+	if (NodeType == ERigMapperFeatureType::WeightedSum)
 	{
-		// Definition for Outputs has to be updated manually since RigMapperOutputFeature is created separately.
-		if (Definition->Outputs.Contains(NodeName))
+		if (URigMapperDefinitionEditorGraphNode_WeightedSum* WsNode =
+			Cast<URigMapperDefinitionEditorGraphNode_WeightedSum>(this))
 		{
-			Definition->Outputs[NodeName] = InNewInput;
+			for (FRigMapperWsInput& Entry : WsNode->WeightedInputs)
+			{
+				if (Entry.Name == InFeatureName)
+				{
+					Entry.Name = TEXT("");
+				}
+			}
 		}
 	}
 	UpdateNodeFields();
-	Definition->BroadcastDefinitionModified();
 }
 
-void URigMapperDefinitionEditorGraphNode::RemoveLinksToFeature(const FString& InFeatureName)
+void URigMapperDefinitionEditorGraphNode::PinConnectionListChanged(UEdGraphPin* Pin)
 {
-	if (!Definition.IsValid())
+	if (!Pin)
 	{
 		return;
 	}
 
-	if (InFeatureName.IsEmpty())
+	if (Pin->Direction == EGPD_Input && Pin->LinkedTo.Num() == 0)
 	{
-		return;
+		RelinkFeature(TEXT(""), Pin);
 	}
+}
 
-	FRigMapperFeature* FeaturePtr = GetNodeFeature();
-	if (!FeaturePtr)
-	{
-		return;
-	}
+void URigMapperDefinitionEditorGraphNode::PreEditChange(FProperty* PropertyAboutToChange)
+{
+	Super::PreEditChange(PropertyAboutToChange);
 
-	FeaturePtr->RenameInput(InFeatureName, TEXT(""));
-	if (NodeType == ERigMapperFeatureType::Output)
+	if (PropertyAboutToChange && PropertyAboutToChange->GetFName() == GET_MEMBER_NAME_CHECKED(URigMapperDefinitionEditorGraphNode, NodeName))
 	{
-		// Definition for Outputs has to be updated manually since RigMapperOutputFeature is created separately.
-		if (Definition->Outputs.Contains(NodeName))
-		{
-			Definition->Outputs[NodeName] = TEXT("");
-		}
+		PreviousNodeName = NodeName;
 	}
-	Definition->BroadcastDefinitionModified();
 }
 
-void URigMapperDefinitionEditorGraphNode::PinConnectionListChanged(UEdGraphPin* Pin)
+void URigMapperDefinitionEditorGraphNode::PostEditChangeProperty(FPropertyChangedEvent& Event)
 {
-	FRigMapperFeature* FeaturePtr = GetNodeFeature();
-	if (!FeaturePtr)
-	{
-		return;
-	}
+	Modify();
 
-	if (!Pin)
+	Super::PostEditChangeProperty(Event);
+
+	URigMapperDefinitionEditorGraph* RigMapperGraph = Cast<URigMapperDefinitionEditorGraph>(GetGraph());
+	if (!RigMapperGraph)
 	{
 		return;
 	}
 
-	if (Pin->Direction == EGPD_Input && Pin->LinkedTo.Num() == 0)
+	// TODO: Refactor with specialized node subclasses for each type
+	if (NodeType == ERigMapperFeatureType::WeightedSum)
 	{
-		RelinkFeature(TEXT(""), Pin);
-	}
-}
+		if (URigMapperDefinitionEditorGraphNode_WeightedSum* WsNode =
+			Cast<URigMapperDefinitionEditorGraphNode_WeightedSum>(this))
+		{
+			TMap<FString, int32> PropertyNameStack;
+			Event.GetArrayIndicesPerObject(0, PropertyNameStack);
 
-FRigMapperFeature* URigMapperDefinitionEditorGraphNode::GetNodeFeature() const
-{
-	if (Definition.IsValid())
+			const FString InputsProperty = GET_MEMBER_NAME_STRING_CHECKED(URigMapperDefinitionEditorGraphNode_WeightedSum, WeightedInputs);
+
+			const int32 PinIndex = PropertyNameStack.Contains(InputsProperty) ? PropertyNameStack.FindRef(InputsProperty) : INDEX_NONE;
+
+			const int32 MinReqPins = GetMinInputPinCount();
+
+			if (Event.ChangeType & EPropertyChangeType::ArrayClear)
+			{
+				// REMOVING ALL INPUTS
+				// All inputs are removed but node has to have minimal number of input pins exposed.
+				for (int32 Index = InputPins.Num() - 1; Index >= MinReqPins; --Index)
+				{
+					RemoveInputPin(InputPins[Index]);
+				}
+				// Break remaining input links.
+				for (UEdGraphPin* InPin : InputPins)
+				{
+					InPin->BreakAllPinLinks();
+				}
+			}
+			else if (Event.ChangeType & EPropertyChangeType::ArrayRemove)
+			{
+				// REMOVING A SINGLE INPUT
+				if (PinIndex != INDEX_NONE && InputPins.IsValidIndex(PinIndex))
+				{
+					if (InputPins.Num() > MinReqPins)
+					{
+						RemoveInputPin(InputPins[PinIndex]);
+					}
+					else
+					{
+						InputPins[PinIndex]->BreakAllPinLinks();
+					}
+				}
+			}
+			else if (Event.ChangeType & EPropertyChangeType::ArrayAdd)
+			{
+				// ADDING A NEW (EMPTY) INPUT
+				if (PinIndex != INDEX_NONE) // This will be true only for WS and Multiply features
+				{
+					TArray<FString> Inputs;
+					if (WsNode->WeightedInputs.Num() > InputPins.Num())
+					{
+						CreateInputPin();
+					}
+				}
+			}
+			else if (Event.ChangeType & EPropertyChangeType::ValueSet)
+			{
+				if ((PinIndex != INDEX_NONE && Event.GetMemberPropertyName() == InputsProperty))
+				{
+					// CHANGING AN INPUT NAME
+					URigMapperDefinitionEditorGraphNode* LinkedNode = nullptr;
+					const int32 InputIndex = PinIndex == INDEX_NONE ? 0 : PinIndex;
+					if (InputPins.IsValidIndex(InputIndex))
+					{
+						InputPins[InputIndex]->BreakAllPinLinks();
+						if (WsNode->WeightedInputs.IsValidIndex(InputIndex))
+						{
+							const FString InputName = WsNode->WeightedInputs[InputIndex].Name;
+							LinkedNode = RigMapperGraph->FindFeatureNodeByName(InputName);
+							if (LinkedNode && LinkedNode->GetOutputPin())
+							{
+								InputPins[InputIndex]->MakeLinkTo(LinkedNode->GetOutputPin());
+							}
+						}
+					}
+				}
+			}
+		}
+	}
+	if (Event.GetPropertyName() == GET_MEMBER_NAME_CHECKED(URigMapperDefinitionEditorGraphNode, NodeName))
 	{
-		return Definition->GetFeature(NodeName, NodeType);
+		if (NodeName != PreviousNodeName)
+		{
+			bool bDuplicateNodeName = false;
+			if (NodeType == ERigMapperFeatureType::Output || NodeType == ERigMapperFeatureType::NullOutput)
+			{
+				bDuplicateNodeName = RigMapperGraph->FindOutputNodeByName(NodeName) != nullptr;
+			}
+			else
+			{
+				bDuplicateNodeName = RigMapperGraph->FindFeatureNodeByName(NodeName) != nullptr;
+			}
+			if (bDuplicateNodeName)
+			{
+				TArray<FString> ExistingNames;
+				RigMapperGraph->GetExistingNodeNames(NodeType, ExistingNames);
+				NodeName = FRigMapperUtils::GenerateUniqueFeatureName(ExistingNames, NodeName, NodeType);
+			}
+
+			if (NodeType != ERigMapperFeatureType::Output && NodeType != ERigMapperFeatureType::NullOutput) // Output and NullOutput nodes don't have output pins.
+			{
+				// Update inputs of all nodes directly connected to this node.
+				for (UEdGraphPin* OutPin : OutputPins) // It will always be one output pin.
+				{
+					for (UEdGraphPin* LinkedPin : OutPin->LinkedTo)
+					{
+						if (URigMapperDefinitionEditorGraphNode* LinkedNode = Cast<URigMapperDefinitionEditorGraphNode>(LinkedPin->GetOwningNode()))
+						{
+							LinkedNode->RelinkFeature(NodeName, LinkedPin);
+						}
+					}
+				}
+			}
+			RigMapperGraph->UpdateNodeNameInMap(PreviousNodeName, NodeName, NodeType);
+		}
 	}
-	return nullptr;
+	UpdateNodeFields();
 }
 
 UEdGraphPin* URigMapperDefinitionEditorGraphNode::CreateInputPin()
@@ -324,206 +314,159 @@ void URigMapperDefinitionEditorGraphNode::GetRect(FVector2D& TopLeft, FVector2D&
 	BottomRight.Y = NodePosY + GetDimensions().Y + GetMargin().Y;
 }
 
-void URigMapperDefinitionEditorGraphNode::UpdateNodeFromFeature(const FPropertyChangedEvent& InPropertyChangedEvent, TWeakPtr<FStructOnScope> InStructContainer)
+void URigMapperDefinitionEditorGraphNode_Input::ApplyToDefinition(URigMapperDefinition* OutDefinition) const
 {
-	if (!InStructContainer.IsValid())
-	{
-		return;
-	}
-	if (!Definition.IsValid())
+	OutDefinition->Inputs.Add(NodeName);
+}
+
+void URigMapperDefinitionEditorGraphNode_WeightedSum::LoadFromFeature(const FRigMapperFeature* InFeature)
+{
+	if (InFeature && InFeature->GetFeatureType() == ERigMapperFeatureType::WeightedSum)
 	{
-		return;
+		const FRigMapperWsFeature* WsFeature = static_cast<const FRigMapperWsFeature*>(InFeature);
+		WeightedInputs.Reset();
+		for (const TPair<FString, double>& Pair : WsFeature->Inputs)
+		{
+			WeightedInputs.Add(FRigMapperWsInput(Pair.Key, Pair.Value));
+		}
+		Range = WsFeature->Range;
 	}
-	URigMapperDefinitionEditorGraph* RigMapperGraph = Cast<URigMapperDefinitionEditorGraph>(GetGraph());
-	if (!RigMapperGraph)
+}
+
+void URigMapperDefinitionEditorGraphNode_WeightedSum::ApplyToDefinition(URigMapperDefinition* OutDefinition) const
+{
+	FRigMapperWsFeature Feature(NodeName);
+	Feature.Range = Range;
+
+	// Rebuild the Inputs TMap from OrderedInputs
+	Feature.Inputs.Empty();
+	for (const FRigMapperWsInput& Input : WeightedInputs)
 	{
-		return;
+		if (!Input.Name.IsEmpty())
+		{
+			Feature.Inputs.Add(Input.Name, Input.Weight);
+		}
 	}
-	FString OldName = NodeName;
-	FRigMapperFeature* FeaturePtr = nullptr;
-	if ((InPropertyChangedEvent.ChangeType & EPropertyChangeType::ValueSet) && InPropertyChangedEvent.GetPropertyName() == GET_MEMBER_NAME_CHECKED(FRigMapperFeature, Name) &&
-		InPropertyChangedEvent.GetMemberPropertyName() == GET_MEMBER_NAME_CHECKED(FRigMapperFeature, Name))
+
+	OutDefinition->Features.WeightedSums.Add(Feature);
+}
+
+void URigMapperDefinitionEditorGraphNode_WeightedSum::UpdateNodeFields()
+{
+	Super::UpdateNodeFields();
+	const FText MinText = LOCTEXT("RangeMinimum", "min");
+	const FText MaxText = LOCTEXT("RangeMaximum", "max");
+
+	for (int32 PinIndex = 0; PinIndex < InputPins.Num(); PinIndex++)
 	{
-		// Name has been modified, find feature using new name.
-		if (const FStrProperty* NameProp = CastField<FStrProperty>(InPropertyChangedEvent.Property))
+		if (WeightedInputs.IsValidIndex(PinIndex))
 		{
-			void* StructMem = InStructContainer.Pin()->GetStructMemory();
-			if (const FString* NamePtr = NameProp->ContainerPtrToValuePtr<FString>(StructMem))
-			{
-				FeaturePtr = Definition->GetFeature(*NamePtr, NodeType);
-			}
+			InputPins[PinIndex]->PinFriendlyName = FText::Format(INVTEXT("{0}"), WeightedInputs[PinIndex].Weight);
 		}
 	}
-	else
+	FString RangeStr;
+	if (Range.bHasLowerBound)
 	{
-		FeaturePtr = GetNodeFeature();
-	}	
-	if (!FeaturePtr)
+		RangeStr.Appendf(TEXT("%s: %.*f"), *MinText.ToString(), 3, Range.LowerBound);
+	}
+	if (Range.bHasLowerBound && Range.bHasUpperBound)
 	{
-		return;
+		RangeStr.Append(TEXT("\n"));
 	}
+	if (Range.bHasUpperBound)
+	{
+		RangeStr.Appendf(TEXT("%s: %.*f"), *MaxText.ToString(), 3, Range.UpperBound);
+	}
+	NodeSubtitle = FText::FromString(RangeStr);
+}
 
-	Modify();
-	NodeName = FeaturePtr->Name;
+void URigMapperDefinitionEditorGraphNode_SDK::LoadFromFeature(const FRigMapperFeature* InFeature)
+{
+	if (InFeature && InFeature->GetFeatureType() == ERigMapperFeatureType::SDK)
+	{
+		const FRigMapperSdkFeature* SdkFeature = static_cast<const FRigMapperSdkFeature*>(InFeature);
+		Keys = SdkFeature->Keys;
+	}
+}
 
-	const FName PropertyName = InPropertyChangedEvent.GetPropertyName();
-	// TODO: Refactor with specialized node subclasses for each type
-	if (NodeType == ERigMapperFeatureType::Multiply || NodeType == ERigMapperFeatureType::WeightedSum || NodeType == ERigMapperFeatureType::SDK || NodeType == ERigMapperFeatureType::Output)
+void URigMapperDefinitionEditorGraphNode_SDK::ApplyToDefinition(URigMapperDefinition* OutDefinition) const
+{
+	FRigMapperSdkFeature Feature(NodeName);
+	Feature.Keys = Keys;
+
+	// Get single input from pin connection
+	const TArray<UEdGraphPin*>& InPins = GetInputPins();
+	if (InPins.Num() > 0 && InPins[0]->LinkedTo.Num() > 0)
 	{
-		TMap<FString, int32> PropertyNameStack;
-		InPropertyChangedEvent.GetArrayIndicesPerObject(0, PropertyNameStack);
+		if (const URigMapperDefinitionEditorGraphNode* LinkedNode =
+			Cast<URigMapperDefinitionEditorGraphNode>(InPins[0]->LinkedTo[0]->GetOwningNode()))
+		{
+			Feature.Input = LinkedNode->NodeName;
+		}
+	}
 
-		const FString InputsProperty = NodeType == ERigMapperFeatureType::WeightedSum ?
-			GET_MEMBER_NAME_STRING_CHECKED(FRigMapperWsFeature, OrderedInputs) :
-			GET_MEMBER_NAME_STRING_CHECKED(FRigMapperMultiplyFeature, Inputs);
+	OutDefinition->Features.SDKs.Add(Feature);
+}
 
-		const int32 PinIndex = PropertyNameStack.Contains(InputsProperty) ? PropertyNameStack.FindRef(InputsProperty) : INDEX_NONE;
+void URigMapperDefinitionEditorGraphNode_SDK::UpdateNodeFields()
+{
+	Super::UpdateNodeFields();
+	TArray<FString> InKeys;
+	TArray<FString> OutKeys;
+	InKeys.Reserve(Keys.Num());
+	OutKeys.Reserve(Keys.Num());
+	for (const FRigMapperSdkKey& Key : Keys)
+	{
+		InKeys.Add(FString::Printf(TEXT("%.*f"), 3, Key.In));
+		OutKeys.Add(FString::Printf(TEXT("%.*f"), 3, Key.Out));
+	}
+	const FString In = FString::Join(InKeys, TEXT(", "));
+	const FString Out = FString::Join(OutKeys, TEXT(", "));
+	NodeSubtitle = FText::FromString(FString::Printf(TEXT("[%s] > [%s]"), *In, *Out));
+}
 
-		const int32 MinReqPins = NodeType == ERigMapperFeatureType::Multiply ? 2 : 1;
+void URigMapperDefinitionEditorGraphNode_Multiply::ApplyToDefinition(URigMapperDefinition* OutDefinition) const
+{
+	FRigMapperMultiplyFeature Feature(NodeName);
 
-		if (InPropertyChangedEvent.ChangeType & EPropertyChangeType::ArrayClear)
-		{	
-			// REMOVING ALL INPUTS
-			// All inputs are removed but node has to have 2 input pin exposed.
-			for (int32 Index = InputPins.Num() - 1; Index >= MinReqPins; --Index)
-			{				
-				RemoveInputPin(InputPins[Index]);
-			}
-			// Break remaining input links.
-			for (UEdGraphPin* InPin : InputPins)
-			{
-				InPin->BreakAllPinLinks();
-			}
-			FeaturePtr->ClearInputs();
-		}
-		else if (InPropertyChangedEvent.ChangeType & EPropertyChangeType::ArrayRemove)
-		{
-			// REMOVING A SINGLE INPUT
-			if (PinIndex != INDEX_NONE && InputPins.IsValidIndex(PinIndex))
-			{
-				if (InputPins.Num() > MinReqPins)
-				{
-					RemoveInputPin(InputPins[PinIndex]);
-				}
-				else
-				{
-					InputPins[PinIndex]->BreakAllPinLinks();
-				}
-				FeaturePtr->UpdateFromDetails();
-			}
-		}
-		else if (InPropertyChangedEvent.ChangeType & EPropertyChangeType::ArrayAdd)
+	// Build input list from pin connections
+	for (const UEdGraphPin* InPin : GetInputPins())
+	{
+		if (InPin->LinkedTo.Num() > 0)
 		{
-			// ADDING A NEW (EMPTY) INPUT
-			if (PinIndex != INDEX_NONE) // This will be true only for WS and Multiply features
+			if (const URigMapperDefinitionEditorGraphNode* LinkedNode =
+				Cast<URigMapperDefinitionEditorGraphNode>(InPin->LinkedTo[0]->GetOwningNode()))
 			{
-				TArray<FString> Inputs;
-				FeaturePtr->GetInputs(Inputs);
-				if (Inputs.Num() > InputPins.Num())
-				{
-					CreateInputPin();
-					FeaturePtr->UpdateFromDetails();
-				}
-			}
-		}
-		else if (InPropertyChangedEvent.ChangeType & EPropertyChangeType::ValueSet)
-		{
-			if ((PinIndex != INDEX_NONE && InPropertyChangedEvent.GetMemberPropertyName() == InputsProperty) ||				// For WS and Multiply features
-				InPropertyChangedEvent.GetPropertyName() == GET_MEMBER_NAME_CHECKED(FRigMapperSdkFeature, Input))		// For SDK and Output features
-			{
-				// CHANGING AN INPUT NAME
-				URigMapperDefinitionEditorGraphNode* LinkedNode = nullptr;
-				TArray<FString> Inputs;
-				FeaturePtr->GetInputs(Inputs);
-				const int32 InputIndex = PinIndex == INDEX_NONE ? 0 : PinIndex;
-				if (InputPins.IsValidIndex(InputIndex))
-				{
-					InputPins[InputIndex]->BreakAllPinLinks();
-					if (Inputs.IsValidIndex(InputIndex))
-					{
-						const FString InputName = Inputs[InputIndex];
-						LinkedNode = RigMapperGraph->FindFeatureNodeByName(InputName);
-						if (LinkedNode)
-						{
-							InputPins[InputIndex]->MakeLinkTo(LinkedNode->GetOutputPin());
-						}
-					}
-				}
-				FeaturePtr->UpdateFromDetails();
-			}
-			else if (PinIndex != INDEX_NONE && NodeType == ERigMapperFeatureType::WeightedSum && InPropertyChangedEvent.GetPropertyName() == GET_MEMBER_NAME_CHECKED(FRigMapperWsInput, Weight))
-			{
-				// CHANGING WS INPUT WEIGHT
-				FeaturePtr->UpdateFromDetails();
+				Feature.Inputs.Add(LinkedNode->NodeName);
 			}
 		}
 	}
 
-	if (NodeName != OldName)
-	{
-		bool bDuplicateNodeName = false;
-		if (NodeType == ERigMapperFeatureType::Output || NodeType == ERigMapperFeatureType::NullOutput)
-		{
-			bDuplicateNodeName = RigMapperGraph->FindOutputNodeByName(NodeName) != nullptr;						
-		}
-		else
-		{
-			bDuplicateNodeName = RigMapperGraph->FindFeatureNodeByName(NodeName) != nullptr;
-		}
-		NodeName = FRigMapperUtils::GenerateUniqueFeatureName(Definition, NodeName, NodeType);
-		// Node name must be unique.
-		// TODO: Inform user about the warning.
-		FeaturePtr->Name = NodeName;
+	OutDefinition->Features.Multiply.Add(Feature);
+}
 
-		if (NodeType != ERigMapperFeatureType::Output && NodeType != ERigMapperFeatureType::NullOutput) // Output and NullOutput nodes don't have output pins.
-		{
-			// Update inputs of all nodes directly connected to this node.
-			for (UEdGraphPin* OutPin : OutputPins) // It will always be one output pin.
-			{
-				for (UEdGraphPin* LinkedPin : OutPin->LinkedTo)
-				{
-					if (URigMapperDefinitionEditorGraphNode* LinkedNode = Cast<URigMapperDefinitionEditorGraphNode>(LinkedPin->GetOwningNode()))
-					{
-						LinkedNode->RelinkFeature(NodeName, LinkedPin);
-					}
-				}
-			}
-		}
-		// Input, Output and NullOutput need to update their Definition data from Feature.
-		if (NodeType == ERigMapperFeatureType::Input)
-		{
-			int32 OutIndex;
-			Definition->Inputs.Find(OldName, OutIndex);
-			if (OutIndex != INDEX_NONE)
-			{
-				Definition->Inputs[OutIndex] = NodeName;
-			}
-		}
-		else if (NodeType == ERigMapperFeatureType::Output)
-		{
-			if (FString* Input = Definition->Outputs.Find(OldName))
-			{
-				FString TempInput = MoveTemp(*Input);
-				Definition->Outputs.Remove(OldName);
-				Definition->Outputs.Add(NodeName, TempInput);
-			}			
-		}
-		else if (NodeType == ERigMapperFeatureType::NullOutput)
+void URigMapperDefinitionEditorGraphNode_Multiply::UpdateNodeFields()
+{
+	Super::UpdateNodeFields();
+}
+
+void URigMapperDefinitionEditorGraphNode_Output::ApplyToDefinition(URigMapperDefinition* OutDefinition) const
+{
+	FString InputName;
+	const TArray<UEdGraphPin*>& InPins = GetInputPins();
+	if (InPins.Num() > 0 && InPins[0]->LinkedTo.Num() > 0)
+	{
+		if (const URigMapperDefinitionEditorGraphNode* LinkedNode =
+			Cast<URigMapperDefinitionEditorGraphNode>(InPins[0]->LinkedTo[0]->GetOwningNode()))
 		{
-			int32 OutIndex;
-			Definition->NullOutputs.Find(OldName, OutIndex);
-			if (OutIndex != INDEX_NONE)
-			{
-				Definition->NullOutputs[OutIndex] = NodeName;
-			}
+			InputName = LinkedNode->NodeName;
 		}
-		// Update node map.
-		RigMapperGraph->UpdateNodeNameInMap(OldName, NodeName, NodeType);
 	}
-		
-	NodeTitle = FText::FromString(NodeName);
-	UpdateNodeFields();
-	Definition->PostEditChangeProperty(const_cast<FPropertyChangedEvent&>(InPropertyChangedEvent));
+	OutDefinition->Outputs.Add(NodeName, InputName);
 }
 
+void URigMapperDefinitionEditorGraphNode_NullOutput::ApplyToDefinition(URigMapperDefinition* OutDefinition) const
+{
+	OutDefinition->NullOutputs.Add(NodeName);
+}
 #undef LOCTEXT_NAMESPACE
diff --git a/Engine/Plugins/Experimental/Animation/RigMapper/Source/RigMapperEditor/Private/RigMapperGraph/RigMapperDefinitionEditorGraphNode.h b/Engine/Plugins/Experimental/Animation/RigMapper/Source/RigMapperEditor/Private/RigMapperGraph/RigMapperDefinitionEditorGraphNode.h
--- a/Engine/Plugins/Experimental/Animation/RigMapper/Source/RigMapperEditor/Private/RigMapperGraph/RigMapperDefinitionEditorGraphNode.h
+++ b/Engine/Plugins/Experimental/Animation/RigMapper/Source/RigMapperEditor/Private/RigMapperGraph/RigMapperDefinitionEditorGraphNode.h
@@ -5,16 +5,11 @@
 #include "CoreMinimal.h"
 #include "EdGraphUtilities.h"
 #include "EdGraph/EdGraphNode.h"
-
+#include "RigMapperDefinition.h"
 #include "RigMapperDefinitionEditorGraphNode.generated.h"
 
 #define UE_API RIGMAPPEREDITOR_API
 
-enum class ERigMapperFeatureType : uint8;
-class URigMapperDefinition;
-class FStructOnScope;
-struct FRigMapperFeature;
-
 /**
  * 
  */
@@ -29,25 +24,33 @@ public:
 		virtual TSharedPtr<SGraphNode> CreateNode(UEdGraphNode* Node) const override;
 	};
 
+	UPROPERTY(EditAnywhere, Category = "Rig Mapper", meta = (DisplayPriority = 1))
+	FString NodeName;
 
 	// UEdGraphNode implementation
 	virtual FText GetNodeTitle(ENodeTitleType::Type TitleType) const override { return NodeTitle; }
-	UE_API virtual FLinearColor GetNodeBodyTintColor() const override;
-	UE_API virtual FLinearColor GetNodeTitleColor() const override;
+	UE_API virtual FLinearColor GetNodeTitleColor() const override { return FLinearColor(0.969, 0.5, 0.466); }
+	UE_API virtual FLinearColor GetNodeBodyTintColor() const override { return GetNodeTitleColor(); }
 	UE_API virtual void GetNodeContextMenuActions(class UToolMenu* Menu, class UGraphNodeContextMenuContext* Context) const override;
 	UE_API virtual void AllocateDefaultPins() override;
 	void PinConnectionListChanged(UEdGraphPin* Pin) override;
+	UE_API virtual void PreEditChange(FProperty* PropertyAboutToChange) override;
+	UE_API virtual void PostEditChangeProperty(FPropertyChangedEvent& Event) override;
+	UE_API virtual void PostLoad() override;
 	// !UEdGraphNode implementation
 
-	UE_API void SetupNode(URigMapperDefinition* InDefinition, const FString& InNodeName, ERigMapperFeatureType InNodeType);
-	const FString& GetNodeName() const { return NodeName; }
+	virtual void LoadFromFeature(const FRigMapperFeature* InFeature) {}
+	virtual void ApplyToDefinition(URigMapperDefinition* OutDefinition) const {}
+	UE_API void SetupNode(const FString& InNodeName, ERigMapperFeatureType InNodeType);
 
 	ERigMapperFeatureType GetNodeType() const { return NodeType; }
 
-	UE_API FRigMapperFeature* GetNodeFeature() const;
-
 	const FText& GetSubtitle() const { return NodeSubtitle; }
 
+	UE_API virtual int32 GetMinInputPinCount() const { return 0; }
+	UE_API virtual int32 GetMaxInputPinCount() const { return MAX_int32; }  // unlimited by default
+	/** Updates node body */
+	UE_API virtual void UpdateNodeFields();
 	UE_API UEdGraphPin* CreateInputPin();
 	UE_API UEdGraphPin* CreateOutputPin();
 	UE_API bool RemoveInputPin(UEdGraphPin* InPinToRemove);
@@ -62,23 +65,15 @@ public:
 
 	UE_API void GetRect(FVector2D& TopLeft, FVector2D& BottomRight) const;
 
-	/** Updates node from definition and makes it dirty */
-	UE_API void UpdateNodeFromFeature(const FPropertyChangedEvent& InPropertyChangedEvent, TWeakPtr<FStructOnScope> InStructContainer);
-
-	/** Updates node body */
-	UE_API void UpdateNodeFields();
-
 	UE_API void RelinkFeature(const FString& InNewInput, UEdGraphPin* InputPin);
 	UE_API void AddNewInputLink(const FString& InNewInput);
 	UE_API void RemoveLinksToFeature(const FString& InFeatureName);
 
-
-
 protected:
 	/** Cached title for the node */
 	FText NodeTitle;
 
-	/** Cached title for the node */
+	/** Cached subtitle for the node */
 	FText NodeSubtitle;
 
 	/** Our input pins */
@@ -93,11 +88,118 @@ protected:
 	/** Cached dimensions of this node (used for layout) */
 	FVector2D Margin = { 0, 0 };
 
-	TWeakObjectPtr<URigMapperDefinition> Definition;
+	UPROPERTY()
+	ERigMapperFeatureType NodeType;
 
-	FString NodeName;
+	FString PreviousNodeName;
+};
 
-	ERigMapperFeatureType NodeType;
+UCLASS(MinimalAPI)
+class URigMapperDefinitionEditorGraphNode_Input : public URigMapperDefinitionEditorGraphNode
+{
+	GENERATED_BODY()
+
+public:
+	UE_API virtual FLinearColor GetNodeTitleColor() const override { return FLinearColor(0.466, 0.969, 0.878); }
+
+	virtual void ApplyToDefinition(URigMapperDefinition* OutDefinition) const override;
+};
+
+USTRUCT(BlueprintType)
+struct FRigMapperWsInput
+{
+	GENERATED_BODY()
+
+public:
+	FRigMapperWsInput() {}
+
+	explicit FRigMapperWsInput(const FString& InName, const double InWeight) : Name(InName), Weight(InWeight) {}
+
+public:
+	UPROPERTY(EditAnywhere, BlueprintReadWrite, Category = "Rig Mapper")
+	FString Name;
+
+	UPROPERTY(EditAnywhere, BlueprintReadWrite, Category = "Rig Mapper")
+	double Weight = 0;
+};
+
+UCLASS(MinimalAPI)
+class URigMapperDefinitionEditorGraphNode_WeightedSum : public URigMapperDefinitionEditorGraphNode
+{
+	GENERATED_BODY()
+
+public:
+	UPROPERTY(EditAnywhere, Category = "Rig Mapper", meta = (DisplayPriority = 1))
+	TArray<FRigMapperWsInput> WeightedInputs;
+
+	UPROPERTY(EditAnywhere, Category = "Rig Mapper", meta = (DisplayPriority = 1))
+	FRigMapperFeatureRange Range;
+
+	UE_API virtual FLinearColor GetNodeTitleColor() const override { return FLinearColor(0.65, 0.466, 0.969); }
+
+	virtual void LoadFromFeature(const FRigMapperFeature* InFeature) override;
+
+	virtual void ApplyToDefinition(URigMapperDefinition* OutDefinition) const override;
+
+	UE_API virtual void UpdateNodeFields() override;
+
+	UE_API virtual int32 GetMinInputPinCount() const { return 1; }
+};
+
+UCLASS(MinimalAPI)
+class URigMapperDefinitionEditorGraphNode_SDK : public URigMapperDefinitionEditorGraphNode
+{
+	GENERATED_BODY()
+
+public:
+	UPROPERTY(EditAnywhere, Category = "Rig Mapper", meta = (DisplayPriority = 1))
+	TArray<FRigMapperSdkKey> Keys;
+
+	UE_API virtual FLinearColor GetNodeTitleColor() const override { return FLinearColor(0.969, 0.717, 0.466); }
+
+	virtual void LoadFromFeature(const FRigMapperFeature* InFeature) override;
+
+	virtual void ApplyToDefinition(URigMapperDefinition* OutDefinition) const override;
+
+	UE_API virtual void UpdateNodeFields() override;
+
+};
+
+UCLASS(MinimalAPI)
+class URigMapperDefinitionEditorGraphNode_Multiply : public URigMapperDefinitionEditorGraphNode
+{
+	GENERATED_BODY()
+
+public:
+	UE_API virtual FLinearColor GetNodeTitleColor() const override { return FLinearColor(0.969, 0.966, 0.466); }
+
+	virtual void ApplyToDefinition(URigMapperDefinition* OutDefinition) const override;
+
+	UE_API virtual void UpdateNodeFields() override;
+
+	UE_API virtual int32 GetMinInputPinCount() const override { return 2; }
+};
+
+UCLASS(MinimalAPI)
+class URigMapperDefinitionEditorGraphNode_Output : public URigMapperDefinitionEditorGraphNode
+{
+	GENERATED_BODY()
+
+public:
+	UE_API virtual FLinearColor GetNodeTitleColor() const override { return FLinearColor(0.466, 0.969, 0.525); }
+
+	virtual void ApplyToDefinition(URigMapperDefinition* OutDefinition) const override;
+};
+
+UCLASS(MinimalAPI)
+class URigMapperDefinitionEditorGraphNode_NullOutput : public URigMapperDefinitionEditorGraphNode
+{
+	GENERATED_BODY()
+
+public:
+	UE_API virtual FLinearColor GetNodeTitleColor() const override { return FLinearColor(0.65, 0.969, 0.466); }
+
+	virtual void ApplyToDefinition(URigMapperDefinition* OutDefinition) const override;
 };
 
 #undef UE_API
diff --git a/Engine/Plugins/Experimental/Animation/RigMapper/Source/RigMapperEditor/Private/RigMapperGraph/RigMapperDefinitionEditorGraphSchema.cpp b/Engine/Plugins/Experimental/Animation/RigMapper/Source/RigMapperEditor/Private/RigMapperGraph/RigMapperDefinitionEditorGraphSchema.cpp
--- a/Engine/Plugins/Experimental/Animation/RigMapper/Source/RigMapperEditor/Private/RigMapperGraph/RigMapperDefinitionEditorGraphSchema.cpp
+++ b/Engine/Plugins/Experimental/Animation/RigMapper/Source/RigMapperEditor/Private/RigMapperGraph/RigMapperDefinitionEditorGraphSchema.cpp
@@ -128,7 +128,7 @@ void URigMapperDefinitionEditorGraphSchema::SetPinBeingDroppedOnNode(UEdGraphPin
 	}
 	if (URigMapperDefinitionEditorGraphNode* SourceNode = Cast<URigMapperDefinitionEditorGraphNode>(InSourcePin->GetOwningNode()))
 	{
-		PinBeingDroppedNodeName = SourceNode->GetNodeName();
+		PinBeingDroppedNodeName = SourceNode->NodeName;
 	}
 }
 
@@ -266,7 +266,6 @@ bool URigMapperDefinitionEditorGraphSchema::TryCreateConnection(UEdGraphPin* InA
 	// We need to save the name of that node in order to update the Definition.
 	UEdGraphPin* InputPin = nullptr;
 	const UEdGraphPin* OutputPin = nullptr;
-	FString OldInputName = TEXT("");
 	FString NewInputName = TEXT("");
 	if (ConnResponse.Response == CONNECT_RESPONSE_BREAK_OTHERS_A)
 	{
@@ -294,7 +293,7 @@ bool URigMapperDefinitionEditorGraphSchema::TryCreateConnection(UEdGraphPin* InA
 	{
 		if (URigMapperDefinitionEditorGraphNode* NewNode = Cast<URigMapperDefinitionEditorGraphNode>(OutputPin->GetOwningNode()))
 		{
-			NewInputName = NewNode->GetNodeName();
+			NewInputName = NewNode->NodeName;
 		}
 	}
 	bModified = Super::TryCreateConnection(InA, InB);
@@ -348,14 +347,14 @@ bool URigMapperDefinitionEditorGraphSchema::SafeDeleteNodeFromGraph(UEdGraph* Gr
 			{
 				if (URigMapperDefinitionEditorGraphNode* LinkedNode = Cast<URigMapperDefinitionEditorGraphNode>(LinkedPin->GetOwningNode()))
 				{
-					LinkedNode->RemoveLinksToFeature(RigMapperNode->GetNodeName());
+					LinkedNode->RemoveLinksToFeature(RigMapperNode->NodeName);
 				}
 			}
 		}
 
 		if (URigMapperDefinitionEditorGraph* RigMapperGraph = Cast<URigMapperDefinitionEditorGraph>(Graph))
 		{
-			RigMapperGraph->RemoveNodeFromMap(RigMapperNode->GetNodeName(), RigMapperNode->GetNodeType());
+			RigMapperGraph->RemoveNodeFromMap(RigMapperNode->NodeName, RigMapperNode->GetNodeType());
 		}
 
 
diff --git a/Engine/Plugins/Experimental/Animation/RigMapper/Source/RigMapperEditor/Private/RigMapperGraph/SRigMapperDefinitionGraphEditor.cpp b/Engine/Plugins/Experimental/Animation/RigMapper/Source/RigMapperEditor/Private/RigMapperGraph/SRigMapperDefinitionGraphEditor.cpp
--- a/Engine/Plugins/Experimental/Animation/RigMapper/Source/RigMapperEditor/Private/RigMapperGraph/SRigMapperDefinitionGraphEditor.cpp
+++ b/Engine/Plugins/Experimental/Animation/RigMapper/Source/RigMapperEditor/Private/RigMapperGraph/SRigMapperDefinitionGraphEditor.cpp
@@ -37,6 +37,8 @@ SRigMapperDefinitionGraphEditor::~SRigMapperDefinitionGraphEditor()
 
 void SRigMapperDefinitionGraphEditor::Construct(const FArguments& InArgs, URigMapperDefinition* InDefinition)
 {
+	check(InDefinition);
+
 	FGraphAppearanceInfo AppearanceInfo;
 	AppearanceInfo.CornerText = LOCTEXT("GraphEditorRigMapperDefinition", "Rig Mapper Definition");
 
@@ -45,11 +47,25 @@ void SRigMapperDefinitionGraphEditor::Construct(const FArguments& InArgs, URigMa
 	GraphEvents.OnSelectionChanged = SGraphEditor::FOnSelectionChanged::CreateSP(this, &SRigMapperDefinitionGraphEditor::HandleSelectionChanged);
 	// GraphEvents.OnNodeDoubleClicked = FSingleNodeEvent::CreateSP(this, &SPhysicsAssetGraph::HandleNodeDoubleClicked);	
 
-	GraphObj = NewObject<URigMapperDefinitionEditorGraph>();
+	URigMapperDefinitionEditorGraph* Graph = Cast<URigMapperDefinitionEditorGraph>(InDefinition->EditorGraph);
+	if (!Graph)
+	{
+		Graph = NewObject<URigMapperDefinitionEditorGraph>(InDefinition, NAME_None, RF_Transactional);
+		InDefinition->EditorGraph = Graph;
+		InDefinition->MarkPackageDirty();
+		Graph->Initialize(InDefinition);
+		Graph->RebuildGraph();
+		Graph->LayoutNodes();
+	}
+	else
+	{
+		Graph->Initialize(InDefinition);
+		Graph->RebuildNodeMaps();
+		Graph->NotifyGraphChanged();
+	}
+
+	GraphObj = Graph;
 	GraphObj->AddToRoot();
-	GraphObj->Initialize(InDefinition);
-	GraphObj->RebuildGraph();
-	GraphObj->LayoutNodes();
 	BindCommands();
 
 	ChildSlot
diff --git a/Engine/Plugins/Experimental/Animation/RigMapper/Source/RigMapperEditor/Private/RigMapperGraph/SRigMapperDefinitionGraphEditor.h b/Engine/Plugins/Experimental/Animation/RigMapper/Source/RigMapperEditor/Private/RigMapperGraph/SRigMapperDefinitionGraphEditor.h
--- a/Engine/Plugins/Experimental/Animation/RigMapper/Source/RigMapperEditor/Private/RigMapperGraph/SRigMapperDefinitionGraphEditor.h
+++ b/Engine/Plugins/Experimental/Animation/RigMapper/Source/RigMapperEditor/Private/RigMapperGraph/SRigMapperDefinitionGraphEditor.h
@@ -36,6 +36,11 @@ public:
 	UE_API void RefreshGraphNode(URigMapperDefinitionEditorGraphNode* InNode);
 	UE_API void CaptureKeyboard();
 
+	TObjectPtr<URigMapperDefinitionEditorGraph> GetEditorGraphObject() const
+	{
+		return GraphObj;
+	}
+
 private:
 	UE_API void ZoomToFitNodes(const TArray<URigMapperDefinitionEditorGraphNode*>& SelectedNodes) const;
 	static UE_API void GetAllLinkedNodes(const URigMapperDefinitionEditorGraphNode* BaseNode, TArray<URigMapperDefinitionEditorGraphNode*>& LinkedNodes, bool bDescend);
