[RigMapper] Implemented support for comment node feature
#Synced-CL 52673514

diff --git a/Engine/Plugins/Experimental/Animation/RigMapper/Source/RigMapperEditor/Private/RigMapperDefinitionEditorToolkit.cpp b/Engine/Plugins/Experimental/Animation/RigMapper/Source/RigMapperEditor/Private/RigMapperDefinitionEditorToolkit.cpp
--- a/Engine/Plugins/Experimental/Animation/RigMapper/Source/RigMapperEditor/Private/RigMapperDefinitionEditorToolkit.cpp
+++ b/Engine/Plugins/Experimental/Animation/RigMapper/Source/RigMapperEditor/Private/RigMapperDefinitionEditorToolkit.cpp
@@ -15,7 +15,6 @@
 #include "RigMapperGraph/SRigMapperDefinitionStructureView.h"
 #include "Framework/Application/SlateApplication.h"
 #include "Widgets/Docking/SDockTab.h"
-#include "Widgets/Views/STreeView.h"
 
 #define LOCTEXT_NAMESPACE "RigMapperDefinitionEditorToolkit"
 
@@ -321,6 +320,16 @@ void FRigMapperDefinitionEditorToolkit::HandleGraphSelectionChanged(const TSet<U
 					this, &FRigMapperDefinitionEditorToolkit::HandleNodePropertyChanged);
 			}
 		}
+		else if (const URigMapperCommentNode* CommentNode = Cast<URigMapperCommentNode>(ObjectNode))
+		{
+			if (Nodes.Num() == 1)
+			{
+				DetailsView->SetObject(const_cast<URigMapperCommentNode*>(CommentNode));
+				DetailsView->OnFinishedChangingProperties().Clear();
+				DetailsView->OnFinishedChangingProperties().AddRaw(
+					this, &FRigMapperDefinitionEditorToolkit::HandleNodePropertyChanged);
+			}
+		}
 	}
 
 	if (GraphTab.IsValid())
diff --git a/Engine/Plugins/Experimental/Animation/RigMapper/Source/RigMapperEditor/Private/RigMapperGraph/RigMapperDefinitionEditorGraphNode.cpp b/Engine/Plugins/Experimental/Animation/RigMapper/Source/RigMapperEditor/Private/RigMapperGraph/RigMapperDefinitionEditorGraphNode.cpp
--- a/Engine/Plugins/Experimental/Animation/RigMapper/Source/RigMapperEditor/Private/RigMapperGraph/RigMapperDefinitionEditorGraphNode.cpp
+++ b/Engine/Plugins/Experimental/Animation/RigMapper/Source/RigMapperEditor/Private/RigMapperGraph/RigMapperDefinitionEditorGraphNode.cpp
@@ -8,6 +8,7 @@
 #include "RigMapperUtils.h"
 #include "PropertyEditorModule.h"
 #include "EdGraph/EdGraphPin.h"
+#include "SGraphNodeComment.h"
 
 #include UE_INLINE_GENERATED_CPP_BY_NAME(RigMapperDefinitionEditorGraphNode)
 
@@ -27,6 +28,10 @@ TSharedPtr<SGraphNode> URigMapperDefinitionEditorGraphNode::NodeFactory::CreateN
 		ThisNode->SetDimensions(GraphNode->GetDesiredSize());		
 		return GraphNode;
 	}
+	if (URigMapperCommentNode* CommentNode = Cast<URigMapperCommentNode>(Node))
+	{
+		return SNew(SGraphNodeComment, CommentNode);
+	}
 	return nullptr;
 }
 
@@ -469,4 +474,13 @@ void URigMapperDefinitionEditorGraphNode_NullOutput::ApplyToDefinition(URigMappe
 {
 	OutDefinition->NullOutputs.Add(NodeName);
 }
+
+void URigMapperCommentNode::OnRenameNode(const FString& NewName)
+{
+	Super::OnRenameNode(NewName);
+	if (UEdGraph* Graph = GetGraph())
+	{
+		Graph->NotifyGraphChanged();
+	}
+}
 #undef LOCTEXT_NAMESPACE
diff --git a/Engine/Plugins/Experimental/Animation/RigMapper/Source/RigMapperEditor/Private/RigMapperGraph/RigMapperDefinitionEditorGraphNode.h b/Engine/Plugins/Experimental/Animation/RigMapper/Source/RigMapperEditor/Private/RigMapperGraph/RigMapperDefinitionEditorGraphNode.h
--- a/Engine/Plugins/Experimental/Animation/RigMapper/Source/RigMapperEditor/Private/RigMapperGraph/RigMapperDefinitionEditorGraphNode.h
+++ b/Engine/Plugins/Experimental/Animation/RigMapper/Source/RigMapperEditor/Private/RigMapperGraph/RigMapperDefinitionEditorGraphNode.h
@@ -6,6 +6,7 @@
 #include "EdGraphUtilities.h"
 #include "EdGraph/EdGraphNode.h"
 #include "RigMapperDefinition.h"
+#include "EdGraphNode_Comment.h"
 #include "RigMapperDefinitionEditorGraphNode.generated.h"
 
 #define UE_API RIGMAPPEREDITOR_API
@@ -202,4 +203,13 @@ public:
 	virtual void ApplyToDefinition(URigMapperDefinition* OutDefinition) const override;
 };
 
+UCLASS(MinimalAPI)
+class URigMapperCommentNode : public UEdGraphNode_Comment
+{
+	GENERATED_BODY()
+public:
+	/** This is needed because CachedCommentTitle is never updated after OnRenameNode like in Blueprint graphs where the Blueprint modification chain triggers automatically. */
+	virtual void OnRenameNode(const FString& NewName) override;
+};
+
 #undef UE_API
diff --git a/Engine/Plugins/Experimental/Animation/RigMapper/Source/RigMapperEditor/Private/RigMapperGraph/RigMapperDefinitionEditorGraphSchema.cpp b/Engine/Plugins/Experimental/Animation/RigMapper/Source/RigMapperEditor/Private/RigMapperGraph/RigMapperDefinitionEditorGraphSchema.cpp
--- a/Engine/Plugins/Experimental/Animation/RigMapper/Source/RigMapperEditor/Private/RigMapperGraph/RigMapperDefinitionEditorGraphSchema.cpp
+++ b/Engine/Plugins/Experimental/Animation/RigMapper/Source/RigMapperEditor/Private/RigMapperGraph/RigMapperDefinitionEditorGraphSchema.cpp
@@ -43,6 +43,26 @@ UEdGraphNode* FSchemaAction_SpawnRigMapperNode::PerformAction(UEdGraph* ParentGr
 	return NewNode;
 }
 
+
+UEdGraphNode* FSchemaAction_AddRigMapperComment::PerformAction(UEdGraph* ParentGraph, UEdGraphPin* FromPin, const FVector2f& Location, bool bSelectNewNode)
+{
+	TSharedPtr<FScopedTransaction> Transaction = MakeShared<FScopedTransaction>(LOCTEXT("RigMapperDefinitionEditor_SpawnNewCommentNode", "Spawn comment node"));
+	ParentGraph->Modify();
+
+	URigMapperCommentNode* CommentNode = NewObject<URigMapperCommentNode>(ParentGraph, NAME_None, RF_Transactional);
+	ParentGraph->AddNode(CommentNode, /*bUserAction=*/true, bSelectNewNode);
+	CommentNode->CreateNewGuid();
+	CommentNode->NodeComment = LOCTEXT("RigMapperDefinitionEditor_CommentDefaultText", "Comment").ToString();
+	CommentNode->PostPlacedNewNode();
+	CommentNode->NodePosX = Location.X;
+	CommentNode->NodePosY = Location.Y;
+	
+	Transaction.Reset();
+
+	return CommentNode;
+}
+
+
 void URigMapperDefinitionEditorGraphSchema::GetContextMenuActions(UToolMenu* Menu, UGraphNodeContextMenuContext* Context) const
 {
 	if (!Context->Pin && Context->Node && Context->Node->IsA<URigMapperDefinitionEditorGraphNode>())
@@ -196,6 +216,9 @@ void URigMapperDefinitionEditorGraphSchema::GetGraphContextActions(FGraphContext
 			LOCTEXT("RigMapperDefinitionEditor_AddNullOutputNodeTooltip", "Spawn null output node"),
 			/*Grouping=*/0));
 	}
+
+	// Add Comment
+	ContextMenuBuilder.AddAction(GetCreateCommentAction());
 }
 
 const FPinConnectionResponse URigMapperDefinitionEditorGraphSchema::CanCreateConnection(const UEdGraphPin* InA, const UEdGraphPin* InB) const
@@ -360,7 +383,19 @@ bool URigMapperDefinitionEditorGraphSchema::SafeDeleteNodeFromGraph(UEdGraph* Gr
 
 		return true;		
 	}
+	if (URigMapperCommentNode* CommentNode = Cast<URigMapperCommentNode>(Node))
+	{
+		// Comment nodes are not kept in the internal node map.
+		return true;
+	}
 	return false;
 }
 
-#undef LOCTEXT_NAMESPACE
+TSharedPtr<FEdGraphSchemaAction> URigMapperDefinitionEditorGraphSchema::GetCreateCommentAction() const
+{
+	return MakeShared<FSchemaAction_AddRigMapperComment>(
+		LOCTEXT("RigMapperDefinitionEditor_AddComment", "Add Comment"),
+		LOCTEXT("RigMapperDefinitionEditor_AddCommentTooltip", "Add a comment box"));
+}
+
+#undef LOCTEXT_NAMESPACE
\ No newline at end of file
diff --git a/Engine/Plugins/Experimental/Animation/RigMapper/Source/RigMapperEditor/Private/RigMapperGraph/RigMapperDefinitionEditorGraphSchema.h b/Engine/Plugins/Experimental/Animation/RigMapper/Source/RigMapperEditor/Private/RigMapperGraph/RigMapperDefinitionEditorGraphSchema.h
--- a/Engine/Plugins/Experimental/Animation/RigMapper/Source/RigMapperEditor/Private/RigMapperGraph/RigMapperDefinitionEditorGraphSchema.h
+++ b/Engine/Plugins/Experimental/Animation/RigMapper/Source/RigMapperEditor/Private/RigMapperGraph/RigMapperDefinitionEditorGraphSchema.h
@@ -29,6 +29,20 @@ private:
 	ERigMapperFeatureType NodeType;
 };
 
+/** Action to add a 'comment' node to the graph */
+USTRUCT()
+struct FSchemaAction_AddRigMapperComment : public FEdGraphSchemaAction
+{
+	GENERATED_BODY()
+
+public:
+	FSchemaAction_AddRigMapperComment() {}
+	FSchemaAction_AddRigMapperComment(const FText& InDescription, const FText& InTooltip)
+		: FEdGraphSchemaAction(FText(), InDescription, InTooltip, 0) {}
+
+	virtual UEdGraphNode* PerformAction(UEdGraph* ParentGraph, UEdGraphPin* FromPin, const FVector2f& Location, bool bSelectNewNode = true) override;
+};
+
 /**
  * 
  */
@@ -78,6 +92,7 @@ public:
 	void SetPinBeingDroppedOnNode(UEdGraphPin* InSourcePin) const override;
 	UE_API bool TryRelinkConnectionTarget(UEdGraphPin* SourcePin, UEdGraphPin* OldTargetPin, UEdGraphPin* NewTargetPin, const TArray<UEdGraphNode*>& InSelectedGraphNodes) const override;
 	bool SafeDeleteNodeFromGraph(UEdGraph* Graph, UEdGraphNode* Node) const override;
+	UE_API virtual TSharedPtr<FEdGraphSchemaAction> GetCreateCommentAction() const override;
 	// End of UEdGraphSchema interface
 
 private:
diff --git a/Engine/Plugins/Experimental/Animation/RigMapper/Source/RigMapperEditor/Private/RigMapperGraph/SRigMapperDefinitionGraphEditor.cpp b/Engine/Plugins/Experimental/Animation/RigMapper/Source/RigMapperEditor/Private/RigMapperGraph/SRigMapperDefinitionGraphEditor.cpp
--- a/Engine/Plugins/Experimental/Animation/RigMapper/Source/RigMapperEditor/Private/RigMapperGraph/SRigMapperDefinitionGraphEditor.cpp
+++ b/Engine/Plugins/Experimental/Animation/RigMapper/Source/RigMapperEditor/Private/RigMapperGraph/SRigMapperDefinitionGraphEditor.cpp
@@ -45,6 +45,14 @@ void SRigMapperDefinitionGraphEditor::Construct(const FArguments& InArgs, URigMa
 	SGraphEditor::FGraphEditorEvents GraphEvents;
 	// GraphEvents.OnCreateActionMenu = SGraphEditor::FOnCreateActionMenu::CreateSP(this, &SPhysicsAssetGraph::OnCreateGraphActionMenu);
 	GraphEvents.OnSelectionChanged = SGraphEditor::FOnSelectionChanged::CreateSP(this, &SRigMapperDefinitionGraphEditor::HandleSelectionChanged);
+	GraphEvents.OnTextCommitted = FOnNodeTextCommitted::CreateLambda(
+		[](const FText& InText, ETextCommit::Type CommitInfo, UEdGraphNode* InNode)
+		{
+			if (InNode && CommitInfo != ETextCommit::OnCleared)
+			{
+				InNode->OnRenameNode(InText.ToString());
+			}
+		});
 	// GraphEvents.OnNodeDoubleClicked = FSingleNodeEvent::CreateSP(this, &SPhysicsAssetGraph::HandleNodeDoubleClicked);	
 
 	URigMapperDefinitionEditorGraph* Graph = Cast<URigMapperDefinitionEditorGraph>(InDefinition->EditorGraph);
@@ -268,11 +276,11 @@ void SRigMapperDefinitionGraphEditor::HandleSelectionChanged(const TSet<UObject*
 			OnSelectionChanged.Execute(Nodes);
 		}
 
-		TArray<URigMapperDefinitionEditorGraphNode*> SelectedNodes;
+		TArray<UEdGraphNode*> SelectedNodes;
 
 		for (UObject* Obj : Nodes)
 		{
-			if (URigMapperDefinitionEditorGraphNode* Node = Cast<URigMapperDefinitionEditorGraphNode>(Obj))
+			if (UEdGraphNode* Node = Cast<UEdGraphNode>(Obj))
 			{
 				SelectedNodes.Add(Node);
 			}
@@ -317,12 +325,12 @@ void SRigMapperDefinitionGraphEditor::OnDeleteSelectedNodes()
 	const FGraphPanelSelectionSet& Selection = GraphEditor->GetSelectedNodes();
 	TSharedPtr<FScopedTransaction> Transaction = MakeShared<FScopedTransaction>(FGenericCommands::Get().Delete->GetDescription());
 
-	TArray<URigMapperDefinitionEditorGraphNode*> NodesToDelete;
+	TArray<UEdGraphNode*> NodesToDelete;
 	NodesToDelete.Reserve(Selection.Num());
 
 	for (UObject* Obj : Selection)
 	{
-		if (URigMapperDefinitionEditorGraphNode* Node = Cast<URigMapperDefinitionEditorGraphNode>(Obj))
+		if (UEdGraphNode* Node = Cast<UEdGraphNode>(Obj))
 		{
 			NodesToDelete.Add(Node);
 		}
@@ -339,10 +347,6 @@ void SRigMapperDefinitionGraphEditor::OnDeleteSelectedNodes()
 		Node->DestroyNode();
 	}
 	GraphEditor->ClearSelectionSet();
-	if (URigMapperDefinitionEditorGraph* RigMapperGraph = Cast<URigMapperDefinitionEditorGraph>(CurrentGraph))
-	{
-		RigMapperGraph->GetDefinition()->BroadcastDefinitionModified();
-	}
 }
 
 bool SRigMapperDefinitionGraphEditor::CanDeleteNodes()
